Copying //UE4/Orion-Staging to //UE4/Main (Source: //Orion/Dev-General @ 3483207)

diff --git a/Engine/Source/Editor/UnrealEd/Private/Fbx/FbxSkeletalMeshImport.cpp b/Engine/Source/Editor/UnrealEd/Private/Fbx/FbxSkeletalMeshImport.cpp
--- a/Engine/Source/Editor/UnrealEd/Private/Fbx/FbxSkeletalMeshImport.cpp
+++ b/Engine/Source/Editor/UnrealEd/Private/Fbx/FbxSkeletalMeshImport.cpp
@@ -59,7 +59,7 @@ using namespace UnFbx;
 
 struct ExistingSkelMeshData;
 extern ExistingSkelMeshData* SaveExistingSkelMeshData(USkeletalMesh* ExistingSkelMesh, bool bSaveMaterials, int32 ReimportLODIndex);
-extern void RestoreExistingSkelMeshData(ExistingSkelMeshData* MeshData, USkeletalMesh* SkeletalMesh, int32 ReimportLODIndex, bool bResetMaterialSlots);
+extern void RestoreExistingSkelMeshData(ExistingSkelMeshData* MeshData, USkeletalMesh* SkeletalMesh, int32 ReimportLODIndex, bool bResetMaterialSlots, bool bIsReimportPreview);
 
 // Get the geometry deformation local to a node. It is never inherited by the
 // children.
@@ -1632,7 +1632,7 @@ USkeletalMesh* UnFbx::FFbxImporter::ImportSkeletalMesh(FImportSkeletalMeshArgs &
 
 		if (ExistSkelMeshDataPtr)
 		{
-			RestoreExistingSkelMeshData(ExistSkelMeshDataPtr, SkeletalMesh, ImportSkeletalMeshArgs.LodIndex, ImportOptions->bResetMaterialSlots);
+			RestoreExistingSkelMeshData(ExistSkelMeshDataPtr, SkeletalMesh, ImportSkeletalMeshArgs.LodIndex, ImportOptions->bResetMaterialSlots, ImportOptions->bIsReimportPreview);
 		}
 
 		SkeletalMesh->CalculateInvRefMatrices();
@@ -1681,10 +1681,17 @@ USkeletalMesh* UnFbx::FFbxImporter::ImportSkeletalMesh(FImportSkeletalMeshArgs &
 			bool bToastSaveMessage = false;
 			if(bFirstMesh || (LastMergeBonesChoice != EAppReturnType::NoAll && LastMergeBonesChoice != EAppReturnType::YesAll))
 			{
-				LastMergeBonesChoice = FMessageDialog::Open(EAppMsgType::YesNoYesAllNoAllCancel,
-					LOCTEXT("SkeletonFailed_BoneMerge", "FAILED TO MERGE BONES:\n\n This could happen if significant hierarchical changes have been made\n"
-						"e.g. inserting a bone between nodes.\nWould you like to regenerate the Skeleton from this mesh?\n\n"
-						"***WARNING: THIS MAY INVALIDATE OR REQUIRE RECOMPRESSION OF ANIMATION DATA.***\n"));
+				if (!ImportOptions->bIsReimportPreview)
+				{
+					LastMergeBonesChoice = FMessageDialog::Open(EAppMsgType::YesNoYesAllNoAllCancel,
+						LOCTEXT("SkeletonFailed_BoneMerge", "FAILED TO MERGE BONES:\n\n This could happen if significant hierarchical changes have been made\n"
+							"e.g. inserting a bone between nodes.\nWould you like to regenerate the Skeleton from this mesh?\n\n"
+							"***WARNING: THIS MAY INVALIDATE OR REQUIRE RECOMPRESSION OF ANIMATION DATA.***\n"));
+				}
+				else
+				{
+					LastMergeBonesChoice = EAppReturnType::NoAll;
+				}
 				bToastSaveMessage = true;
 			}
 			
@@ -3752,13 +3759,12 @@ public:
 	void ComputeMorphDeltas()
 	{
 		TArray<bool> WasProcessed;
-		WasProcessed.Empty( BaseImportData.Wedges.Num() );
-		WasProcessed.AddZeroed( BaseImportData.Wedges.Num() );
+		WasProcessed.Empty(LODModel->NumVertices );
+		WasProcessed.AddZeroed(LODModel->NumVertices);
 
 		for ( int32 Idx = 0; Idx < BaseIndexData.Indices.Num(); ++Idx )
 		{
 			uint32 BaseVertIdx = BaseIndexData.Indices[Idx];
-
 			// check for duplicate processing
 			if( !WasProcessed[BaseVertIdx] )
 			{
@@ -3770,42 +3776,44 @@ public:
 				{
 					// get the base mesh's original wedge point index
 					uint32 BasePointIdx = BaseWedgePointIndices[ BaseVertIdx ];
+					if (BaseImportData.Points.IsValidIndex(BasePointIdx))
+					{
+						FVector BasePosition = BaseImportData.Points[BasePointIdx];
+						FVector TargetPosition = MorphLODPoints[BasePointIdx];
 
-					FVector BasePosition = BaseImportData.Points[ BasePointIdx ];
-					FVector TargetPosition = MorphLODPoints[ BasePointIdx ];
-
-					FVector PositionDelta = TargetPosition - BasePosition;
+						FVector PositionDelta = TargetPosition - BasePosition;
 
-					uint32* VertexIdx = WedgePointToVertexIndexMap.Find( BasePointIdx );
+						uint32* VertexIdx = WedgePointToVertexIndexMap.Find(BasePointIdx);
 
-					FVector NormalDeltaZ = FVector::ZeroVector;
+						FVector NormalDeltaZ = FVector::ZeroVector;
 
-					if ( VertexIdx != nullptr )
-					{
-						FVector BaseNormal = BaseTangentZ[ *VertexIdx ];
-						FVector TargetNormal = TangentZ[ *VertexIdx ];
+						if (VertexIdx != nullptr)
+						{
+							FVector BaseNormal = BaseTangentZ[*VertexIdx];
+							FVector TargetNormal = TangentZ[*VertexIdx];
 
-						NormalDeltaZ = TargetNormal - BaseNormal;
-					}
+							NormalDeltaZ = TargetNormal - BaseNormal;
+						}
 
-					// check if position actually changed much
-					if( PositionDelta.SizeSquared() > FMath::Square(THRESH_POINTS_ARE_NEAR) || 
-						// since we can't get imported morphtarget normal from FBX
-						// we can't compare normal unless it's calculated
-						// this is special flag to ignore normal diff
-						( (ImportOptions->ShouldImportNormals() == false) && NormalDeltaZ.SizeSquared() > 0.01f) )
-					{
-						// create a new entry
-						FMorphTargetDelta NewVertex;
-						// position delta
-						NewVertex.PositionDelta = PositionDelta;
-						// normal delta
-						NewVertex.TangentZDelta = NormalDeltaZ;
-						// index of base mesh vert this entry is to modify
-						NewVertex.SourceIdx = BaseVertIdx;
-
-						// add it to the list of changed verts
-						MorphTargetDeltas.Add( NewVertex );
+						// check if position actually changed much
+						if (PositionDelta.SizeSquared() > FMath::Square(THRESH_POINTS_ARE_NEAR) ||
+							// since we can't get imported morphtarget normal from FBX
+							// we can't compare normal unless it's calculated
+							// this is special flag to ignore normal diff
+							((ImportOptions->ShouldImportNormals() == false) && NormalDeltaZ.SizeSquared() > 0.01f))
+						{
+							// create a new entry
+							FMorphTargetDelta NewVertex;
+							// position delta
+							NewVertex.PositionDelta = PositionDelta;
+							// normal delta
+							NewVertex.TangentZDelta = NormalDeltaZ;
+							// index of base mesh vert this entry is to modify
+							NewVertex.SourceIdx = BaseVertIdx;
+
+							// add it to the list of changed verts
+							MorphTargetDeltas.Add(NewVertex);
+						}
 					}
 				}
 			}
diff --git a/Engine/Source/Editor/UnrealEd/Private/FbxCompareWindow.h b/Engine/Source/Editor/UnrealEd/Private/FbxCompareWindow.h
new file mode 100644
--- /dev/null
+++ b/Engine/Source/Editor/UnrealEd/Private/FbxCompareWindow.h
@@ -0,0 +1,390 @@
+// Copyright 1998-2017 Epic Games, Inc. All Rights Reserved.
+
+#pragma once
+
+#include "CoreMinimal.h"
+#include "InputCoreTypes.h"
+#include "Widgets/DeclarativeSyntaxSupport.h"
+#include "Input/Reply.h"
+#include "Widgets/SCompoundWidget.h"
+#include "Widgets/SWindow.h"
+#include "Widgets/Views/STableViewBase.h"
+#include "Widgets/Views/STableRow.h"
+#include "Widgets/Views/STreeView.h"
+#include "Widgets/SNullWidget.h"
+#include "EditorStyleSet.h"
+
+class FFbxSceneInfo;
+
+enum EFBXCompareSection
+{
+	EFBXCompareSection_General = 0,
+	EFBXCompareSection_Materials,
+	EFBXCompareSection_Skeleton,
+	EFBXCompareSection_Count
+};
+
+class FCompMaterial
+{
+public:
+	FCompMaterial(FName InMaterialSlotName, FName InImportedMaterialSlotName)
+		: MaterialSlotName(InMaterialSlotName)
+		, ImportedMaterialSlotName(InImportedMaterialSlotName)
+	{}
+	FName MaterialSlotName;
+	FName ImportedMaterialSlotName;
+};
+
+class FCompSection
+{
+public:
+	FCompSection()
+		: MaterialIndex(INDEX_NONE)
+	{}
+
+	int32 MaterialIndex;
+};
+
+class FCompLOD
+{
+public:
+	~FCompLOD()
+	{
+		Sections.Empty();
+	}
+
+	TArray<FCompSection> Sections;
+};
+
+class FCompJoint
+{
+public:
+	FCompJoint()
+		: Name(NAME_None)
+		, ParentIndex(INDEX_NONE)
+	{ }
+
+	~FCompJoint()
+	{
+		ChildIndexes.Empty();
+	}
+
+	FName Name;
+	int32 ParentIndex;
+	TArray<int32> ChildIndexes;
+};
+
+class FCompSkeleton
+{
+public:
+	FCompSkeleton()
+		: bSkeletonFitMesh(true)
+	{}
+
+	~FCompSkeleton()
+	{
+		Joints.Empty();
+	}
+
+	TArray<FCompJoint> Joints;
+	bool bSkeletonFitMesh;
+};
+
+class FCompMesh
+{
+public:
+	~FCompMesh()
+	{
+		CompMaterials.Empty();
+		CompLods.Empty();
+	}
+
+	TArray<FCompMaterial> CompMaterials;
+	TArray<FCompLOD> CompLods;
+	FCompSkeleton CompSkeleton;
+	TArray<FString> UVSetsName;
+	int32 LightMapUvIndex;
+
+	TArray<FString> ErrorMessages;
+	TArray<FString> WarningMessages;
+};
+
+class FGeneralFbxFileInfo
+{
+public:
+	FString ApplicationCreator;
+	FString UE4SdkVersion;
+	FString FileVersion;
+	FString AxisSystem;
+	FString UnitSystem;
+	FString CreationDate;
+};
+
+class FSkeletonCompareData : public TSharedFromThis<FSkeletonCompareData>
+{
+public:
+	FSkeletonCompareData()
+		: CurrentJointIndex(INDEX_NONE)
+		, FbxJointIndex(INDEX_NONE)
+		, JointName(NAME_None)
+		, ParentJoint(nullptr)
+		, bMatchJoint(false)
+		, bChildConflict(false)
+	{}
+	int32 CurrentJointIndex;
+	int32 FbxJointIndex;
+	FName JointName;
+	TSharedPtr<FSkeletonCompareData> ParentJoint;
+	bool bMatchJoint;
+	bool bChildConflict;
+	TArray<int32> ChildJointIndexes;
+	TArray<TSharedPtr<FSkeletonCompareData>> ChildJoints;
+};
+
+class FCompareRowData : public TSharedFromThis<FCompareRowData>
+{
+public:
+	FCompareRowData()
+		: RowIndex(INDEX_NONE)
+		, CurrentData(nullptr)
+		, FbxData(nullptr)
+	{}
+	virtual ~FCompareRowData() {}
+
+	int32 RowIndex;
+	FCompMesh *CurrentData;
+	FCompMesh *FbxData;
+
+	virtual TSharedRef<SWidget> ConstructCellCurrent()
+	{
+		return SNullWidget::NullWidget;
+	};
+	virtual TSharedRef<SWidget> ConstructCellFbx()
+	{
+		return SNullWidget::NullWidget;
+	};
+};
+
+class FMaterialCompareData : public FCompareRowData
+{
+public:
+	enum EMaterialCompareDisplayOption
+	{
+		NoMatch = 0,
+		IndexChanged,
+		SkinxxError,
+		All,
+		MaxOptionEnum
+	};
+
+	FMaterialCompareData()
+		: CurrentMaterialIndex(INDEX_NONE)
+		, FbxMaterialIndex(INDEX_NONE)
+		, CurrentMaterialMatch(INDEX_NONE)
+		, FbxMaterialMatch(INDEX_NONE)
+		, bCurrentSkinxxDuplicate(false)
+		, bCurrentSkinxxMissing(false)
+		, bFbxSkinxxDuplicate(false)
+		, bFbxSkinxxMissing(false)
+		, CompareDisplayOption(EMaterialCompareDisplayOption::All)
+	{}
+	
+	virtual ~FMaterialCompareData() {}
+
+	int32 CurrentMaterialIndex;
+	int32 FbxMaterialIndex;
+	int32 CurrentMaterialMatch;
+	int32 FbxMaterialMatch;
+	bool bCurrentSkinxxDuplicate;
+	bool bCurrentSkinxxMissing;
+	bool bFbxSkinxxDuplicate;
+	bool bFbxSkinxxMissing;
+	EMaterialCompareDisplayOption CompareDisplayOption;
+	
+	FSlateColor GetCellColor(FCompMesh *DataA, int32 MaterialIndexA, int32 MaterialMatchA, FCompMesh *DataB, int32 MaterialIndexB, bool bSkinxxError) const;
+	FSlateColor GetCurrentCellColor() const;
+	FSlateColor GetFbxCellColor() const;
+
+	TSharedRef<SWidget> ConstructCell(FCompMesh *MeshData, int32 MeshMaterialIndex, bool bSkinxxDuplicate, bool bSkinxxMissing);
+	virtual TSharedRef<SWidget> ConstructCellCurrent() override;
+	virtual TSharedRef<SWidget> ConstructCellFbx() override;
+};
+
+class SCompareRowDataTableListViewRow : public SMultiColumnTableRow<TSharedPtr<FCompareRowData>>
+{
+public:
+
+	SLATE_BEGIN_ARGS(SCompareRowDataTableListViewRow)
+		: _CompareRowData(nullptr)
+	{}
+
+	/** The item content. */
+	SLATE_ARGUMENT(TSharedPtr<FCompareRowData>, CompareRowData)
+	SLATE_END_ARGS()
+
+	void Construct(const FArguments& InArgs, const TSharedRef<STableViewBase>& InOwnerTableView)
+	{
+		CompareRowData = InArgs._CompareRowData;
+
+		//This is suppose to always be valid
+		check(CompareRowData.IsValid());
+
+		SMultiColumnTableRow<TSharedPtr<FCompareRowData>>::Construct(
+			FSuperRowType::FArguments()
+			.Style(FEditorStyle::Get(), "DataTableEditor.CellListViewRow"),
+			InOwnerTableView
+			);
+	}
+
+	/** Overridden from SMultiColumnTableRow.  Generates a widget for this column of the list view. */
+	virtual TSharedRef<SWidget> GenerateWidgetForColumn(const FName& ColumnName) override
+	{
+		if (ColumnName == FName(TEXT("RowIndex")))
+		{
+			return SNew(SBox)
+				.Padding(FMargin(5.0f, 0.0f, 0.0f, 0.0f))
+				[
+					SNew(STextBlock)
+					.Text(FText::FromString(FString::FromInt(CompareRowData->RowIndex)))
+				];
+		}
+		else if (ColumnName == FName(TEXT("Current")))
+		{
+			return CompareRowData->ConstructCellCurrent();
+		}
+		else if (ColumnName == FName(TEXT("Fbx")))
+		{
+			return CompareRowData->ConstructCellFbx();
+		}
+		return SNullWidget::NullWidget;
+	}
+private:
+
+	/** The node info to build the tree view row from. */
+	TSharedPtr<FCompareRowData> CompareRowData;
+};
+
+class SFbxCompareWindow : public SCompoundWidget
+{
+public:
+	SLATE_BEGIN_ARGS(SFbxCompareWindow)
+		: _WidgetWindow()
+		, _FullFbxPath()
+		, _FbxSceneInfo()
+		, _FbxGeneralInfo()
+		, _AssetReferencingSkeleton(nullptr)
+		, _CurrentMeshData(nullptr)
+		, _FbxMeshData(nullptr)
+		, _PreviewObject(nullptr)
+		{}
+
+		SLATE_ARGUMENT( TSharedPtr<SWindow>, WidgetWindow )
+		SLATE_ARGUMENT( FText, FullFbxPath )
+		SLATE_ARGUMENT( TSharedPtr<FFbxSceneInfo>, FbxSceneInfo )
+		SLATE_ARGUMENT( FGeneralFbxFileInfo, FbxGeneralInfo )
+		SLATE_ARGUMENT( TArray<TSharedPtr<FString>>*, AssetReferencingSkeleton)
+		SLATE_ARGUMENT( FCompMesh*, CurrentMeshData )
+		SLATE_ARGUMENT( FCompMesh*, FbxMeshData )
+		SLATE_ARGUMENT( UObject*, PreviewObject )
+			
+	SLATE_END_ARGS()
+
+public:
+	void Construct(const FArguments& InArgs);
+	virtual bool SupportsKeyboardFocus() const override { return true; }
+
+	FReply OnDone()
+	{
+		if ( WidgetWindow.IsValid() )
+		{
+			WidgetWindow.Pin()->RequestDestroyWindow();
+		}
+		return FReply::Handled();
+	}
+
+	virtual FReply OnKeyDown( const FGeometry& MyGeometry, const FKeyEvent& InKeyEvent ) override
+	{
+		if( InKeyEvent.GetKey() == EKeys::Escape )
+		{
+			return OnDone();
+		}
+
+		return FReply::Unhandled();
+	}
+
+	SFbxCompareWindow()
+		: FullFbxPath(TEXT(""))
+		, FbxSceneInfo()
+		
+	{}
+		
+private:
+	TWeakPtr< SWindow > WidgetWindow;
+	FString FullFbxPath;
+
+	//Preview Mesh
+	UObject *PreviewObject;
+
+	//////////////////////////////////////////////////////////////////////////
+	//Collapse generic
+	bool bShowSectionFlag[EFBXCompareSection_Count];
+	FReply SetSectionVisible(EFBXCompareSection SectionIndex);
+	EVisibility IsSectionVisible(EFBXCompareSection SectionIndex);
+	const FSlateBrush* GetCollapsableArrow(EFBXCompareSection SectionIndex) const;
+	//////////////////////////////////////////////////////////////////////////
+	
+	//////////////////////////////////////////////////////////////////////////
+	// General fbx data
+	FGeneralFbxFileInfo FbxGeneralInfo;
+	TSharedPtr<FFbxSceneInfo> FbxSceneInfo;
+	TArray<TSharedPtr<FString>> GeneralListItem;
+	
+	void FillGeneralListItem();
+	//Construct slate
+	TSharedPtr<SWidget> ConstructGeneralInfo();
+	//Slate events
+	TSharedRef<ITableRow> OnGenerateRowGeneralFbxInfo(TSharedPtr<FString> InItem, const TSharedRef<STableViewBase>& OwnerTable);
+	//////////////////////////////////////////////////////////////////////////
+
+	//////////////////////////////////////////////////////////////////////////
+	// Compare data
+	FCompMesh *CurrentMeshData;
+	FCompMesh *FbxMeshData;
+	//////////////////////////////////////////////////////////////////////////
+
+	//////////////////////////////////////////////////////////////////////////
+	// Material Data
+	TArray<TSharedPtr<FMaterialCompareData>> CompareMaterialListItem;
+	FMaterialCompareData::EMaterialCompareDisplayOption CurrentDisplayOption;
+	
+	void FillMaterialListItem();
+	bool FindSkinxxErrors(FCompMesh *MeshData, TArray<bool> &DuplicateSkinxxMaterialNames, TArray<bool> &MissingSkinxxSuffixeMaterialNames);
+	//Construct slate
+	TSharedPtr<SWidget> ConstructMaterialComparison();
+	//Slate events
+	TSharedRef<ITableRow> OnGenerateRowForCompareMaterialList(TSharedPtr<FMaterialCompareData> RowData, const TSharedRef<STableViewBase>& Table);
+	void ToggleMaterialDisplay(ECheckBoxState InNewValue, FMaterialCompareData::EMaterialCompareDisplayOption InDisplayOption);
+	ECheckBoxState IsToggleMaterialDisplayChecked(FMaterialCompareData::EMaterialCompareDisplayOption InDisplayOption) const;
+	//////////////////////////////////////////////////////////////////////////
+
+	//////////////////////////////////////////////////////////////////////////
+	// Skeleton Data
+	TSharedPtr<STreeView<TSharedPtr<FSkeletonCompareData>>> CompareTree;
+	TArray<TSharedPtr<FSkeletonCompareData>> DisplaySkeletonTreeItem;
+
+	TArray<TSharedPtr<FSkeletonCompareData>> CurrentSkeletonTreeItem;
+	TArray<TSharedPtr<FSkeletonCompareData>> FbxSkeletonTreeItem;
+
+	TArray<TSharedPtr<FString>> AssetReferencingSkeleton;
+
+	void FilSkeletonTreeItem();
+	void RecursiveMatchJointInfo(TSharedPtr<FSkeletonCompareData> CurrentItem);
+	void SetMatchJointInfo();
+	//Construct slate
+	TSharedPtr<SWidget> ConstructSkeletonComparison();
+	//Slate events
+	TSharedRef<ITableRow> OnGenerateRowCompareTreeView(TSharedPtr<FSkeletonCompareData> RowData, const TSharedRef<STableViewBase>& Table);
+	void OnGetChildrenRowCompareTreeView(TSharedPtr<FSkeletonCompareData> InParent, TArray< TSharedPtr<FSkeletonCompareData> >& OutChildren);
+	TSharedRef<ITableRow> OnGenerateRowAssetReferencingSkeleton(TSharedPtr<FString> InItem, const TSharedRef<STableViewBase>& OwnerTable);
+	//////////////////////////////////////////////////////////////////////////
+};
diff --git a/Engine/Source/Editor/UnrealEd/Private/FbxImporter.h b/Engine/Source/Editor/UnrealEd/Private/FbxImporter.h
--- a/Engine/Source/Editor/UnrealEd/Private/FbxImporter.h
+++ b/Engine/Source/Editor/UnrealEd/Private/FbxImporter.h
@@ -187,6 +187,9 @@ struct FBXImportOptions
 	//This data allow to override some fbx Material(point by the uint64 id) with existing unreal material asset
 	TMap<uint64, class UMaterialInterface*> OverrideMaterials;
 
+	//The importer is importing a preview
+	bool bIsReimportPreview;
+
 	bool ShouldImportNormals()
 	{
 		return NormalImportMethod == FBXNIM_ImportNormals || NormalImportMethod == FBXNIM_ImportNormalsAndTangents;
@@ -468,7 +471,7 @@ private:
 	static FbxAMatrix JointPostConversionMatrix;
 };
 
-FBXImportOptions* GetImportOptions( class FFbxImporter* FbxImporter, UFbxImportUI* ImportUI, bool bShowOptionDialog, bool bIsAutomated, const FString& FullPath, bool& OutOperationCanceled, bool& OutImportAll, bool bIsObjFormat, bool bForceImportType = false, EFBXImportType ImportType = FBXIT_StaticMesh );
+FBXImportOptions* GetImportOptions( class FFbxImporter* FbxImporter, UFbxImportUI* ImportUI, bool bShowOptionDialog, bool bIsAutomated, const FString& FullPath, bool& OutOperationCanceled, bool& OutImportAll, bool bIsObjFormat, bool bForceImportType = false, EFBXImportType ImportType = FBXIT_StaticMesh, UObject* ReimportObject = nullptr);
 void ApplyImportUIToImportOptions(UFbxImportUI* ImportUI, FBXImportOptions& InOutImportOptions);
 
 struct FImportedMaterialData
@@ -503,6 +506,9 @@ public:
 	UNREALED_API static FFbxImporter* GetInstance();
 	static void DeleteInstance();
 
+	static FFbxImporter* GetPreviewInstance();
+	static void DeletePreviewInstance();
+
 	/**
 	 * Detect if the FBX file has skeletal mesh model. If there is deformer definition, then there is skeletal mesh.
 	 * In this function, we don't need to import the scene. But the open process is time-consume if the file is large.
@@ -900,6 +906,16 @@ public:
 	 */
 	UNREALED_API FBXImportOptions* GetImportOptions() const;
 
+	/*
+	* This function show a dialog to let the user know what will be change if the fbx is imported
+	*/
+	void ShowFbxReimportPreview(UObject *ReimportObj, UFbxImportUI* ImportUI, const FString& FullPath);
+
+	/*
+	* Function use to retrieve general fbx information for the preview
+	*/
+	void FillGeneralFbxFileInformation(void *GeneralInfoPtr);
+
 	/** helper function **/
 	UNREALED_API static void DumpFBXNode(FbxNode* Node);
 
@@ -1041,6 +1057,24 @@ public:
 	//We cache the hash of the file when we open the file. This is to avoid calculating the hash many time when importing many asset in one fbx file.
 	FMD5Hash Md5Hash;
 
+	struct FFbxMaterial
+	{
+		FbxSurfaceMaterial* FbxMaterial;
+		UMaterialInterface* Material;
+
+		FString GetName() const { return FbxMaterial ? ANSI_TO_TCHAR(FbxMaterial->GetName()) : TEXT("None"); }
+	};
+
+	/**
+	* Make material Unreal asset name from the Fbx material
+	*
+	* @param FbxMaterial Material from the Fbx node
+	* @return Sanitized asset name
+	*/
+	FString GetMaterialFullName(FbxSurfaceMaterial& FbxMaterial);
+
+	FbxGeometryConverter* GetGeometryConverter() { return GeometryConverter; }
+
 protected:
 	enum IMPORTPHASE
 	{
@@ -1050,14 +1084,7 @@ protected:
 	};
 	
 	static TSharedPtr<FFbxImporter> StaticInstance;
-
-	struct FFbxMaterial
-	{
-		FbxSurfaceMaterial* FbxMaterial;
-		UMaterialInterface* Material;
-
-		FString GetName() const { return FbxMaterial ? ANSI_TO_TCHAR( FbxMaterial->GetName() ) : TEXT("None"); }
-	};
+	static TSharedPtr<FFbxImporter> StaticPreviewInstance;
 	
 	//make sure we are not applying two time the option transform to the same node
 	TArray<FbxNode*> TransformSettingsToFbxApply;
@@ -1074,6 +1101,10 @@ protected:
 	TWeakObjectPtr<UObject> Parent;
 	FString FbxFileVersion;
 
+	//Original File Info
+	FbxAxisSystem FileAxisSystem;
+	FbxSystemUnit FileUnitSystem;
+
 	// Flag that the mesh is the first mesh to import in current FBX scene
 	// FBX scene may contain multiple meshes, importer can import them at one time.
 	// Initialized as true when start to import a FBX scene
@@ -1178,6 +1209,8 @@ protected:
 	*/
     bool FillSkelMeshImporterFromFbx(FSkeletalMeshImportData& ImportData, FbxMesh*& Mesh, FbxSkin* Skin, 
 										FbxShape* Shape, TArray<FbxNode*> &SortedLinks, const TArray<FbxSurfaceMaterial*>& FbxMaterials, FbxNode *RootNode);
+public:
+
 	/**
 	* Fill FSkeletalMeshIMportData from Fbx Nodes and FbxShape Array if exists.  
 	*
@@ -1189,7 +1222,8 @@ protected:
 	* @returns bool*	true if import successfully.
 	*/
 	bool FillSkeletalMeshImportData(TArray<FbxNode*>& NodeArray, UFbxSkeletalMeshImportData* TemplateImportData, TArray<FbxShape*> *FbxShapeArray, FSkeletalMeshImportData* OutData, TArray<FName> &LastImportedMaterialNames);
-	
+
+protected:
 	/**
 	* Fill the Points in FSkeletalMeshIMportData from a Fbx Node and a FbxShape if it exists.
 	*
@@ -1329,14 +1363,6 @@ protected:
 	 */
 	int32 CreateNodeMaterials(FbxNode* FbxNode, TArray<UMaterialInterface*>& outMaterials, TArray<FString>& UVSets, bool bForSkeletalMesh);
 
-	/**
-	 * Make material Unreal asset name from the Fbx material
-	 *
-	 * @param FbxMaterial Material from the Fbx node
-	 * @return Sanitized asset name
-	 */
-	FString GetMaterialFullName(FbxSurfaceMaterial& FbxMaterial);
-
 	/**
 	 * Create Unreal material from Fbx material.
 	 * Only setup channels that connect to texture, and setup the UV coordinate of texture.
diff --git a/Engine/Source/Editor/UnrealEd/Private/FbxMeshUtils.cpp b/Engine/Source/Editor/UnrealEd/Private/FbxMeshUtils.cpp
--- a/Engine/Source/Editor/UnrealEd/Private/FbxMeshUtils.cpp
+++ b/Engine/Source/Editor/UnrealEd/Private/FbxMeshUtils.cpp
@@ -44,7 +44,7 @@ extern void UpdateSomeLodsImportMeshData(UStaticMesh* NewMesh, TArray<int32> *Re
 
 struct ExistingSkelMeshData;
 extern ExistingSkelMeshData* SaveExistingSkelMeshData(USkeletalMesh* ExistingSkelMesh, bool bSaveMaterials, int32 ReimportLODIndex);
-extern void RestoreExistingSkelMeshData(ExistingSkelMeshData* MeshData, USkeletalMesh* SkeletalMesh, int32 ReimportLODIndex, bool bResetMaterialSlots);
+extern void RestoreExistingSkelMeshData(ExistingSkelMeshData* MeshData, USkeletalMesh* SkeletalMesh, int32 ReimportLODIndex, bool bResetMaterialSlots, bool bIsReimportPreview);
 
 namespace FbxMeshUtils
 {
@@ -460,7 +460,7 @@ namespace FbxMeshUtils
 
 						if (SkelMeshDataPtr != nullptr)
 						{
-							RestoreExistingSkelMeshData(SkelMeshDataPtr, SelectedSkelMesh, SelectedLOD, false);
+							RestoreExistingSkelMeshData(SkelMeshDataPtr, SelectedSkelMesh, SelectedLOD, false, ImportOptions->bIsReimportPreview);
 						}
 						SelectedSkelMesh->PostEditChange();
 						// Mark package containing skeletal mesh as dirty.
diff --git a/Engine/Source/Editor/UnrealEd/Private/FbxOptionWindow.h b/Engine/Source/Editor/UnrealEd/Private/FbxOptionWindow.h
--- a/Engine/Source/Editor/UnrealEd/Private/FbxOptionWindow.h
+++ b/Engine/Source/Editor/UnrealEd/Private/FbxOptionWindow.h
@@ -32,6 +32,7 @@ public:
 		SLATE_ARGUMENT( bool, IsObjFormat )
 		SLATE_ARGUMENT( float, MaxWindowHeight)
 		SLATE_ARGUMENT(float, MaxWindowWidth)
+		SLATE_EVENT(FOnPreviewFbxImport, OnPreviewFbxImport)
 	SLATE_END_ARGS()
 
 public:
@@ -94,13 +95,18 @@ public:
 private:
 
 	bool CanImport() const;
+	FReply OnPreviewClick() const;
+	FReply OnResetToDefaultClick() const;
 	FText GetImportTypeDisplayText() const;
 
 private:
 	UFbxImportUI*	ImportUI;
+	TSharedPtr<class IDetailsView> DetailsView;
 	TWeakPtr< SWindow > WidgetWindow;
 	TSharedPtr< SButton > ImportButton;
 	bool			bShouldImport;
 	bool			bShouldImportAll;
 	bool			bIsObjFormat;
+
+	FOnPreviewFbxImport OnPreviewFbxImport;
 };
diff --git a/Engine/Source/Editor/UnrealEd/Private/UnrealEdEngine.cpp b/Engine/Source/Editor/UnrealEd/Private/UnrealEdEngine.cpp
--- a/Engine/Source/Editor/UnrealEd/Private/UnrealEdEngine.cpp
+++ b/Engine/Source/Editor/UnrealEd/Private/UnrealEdEngine.cpp
@@ -160,8 +160,14 @@ void UUnrealEdEngine::Init(IEngineLoop* InEngineLoop)
 		BaseCookingFlags |= CookerSettings->bIterativeCookingForLaunchOn ? IterativeFlags : ECookInitializationFlags::None;
 		BaseCookingFlags |= CookerSettings->bEnableBuildDDCInBackground ? ECookInitializationFlags::BuildDDCInBackground : ECookInitializationFlags::None;
 
+
 		if (CookerSettings->bEnableCookOnTheSide)
 		{
+			if ( ExperimentalSettings->bSharedCookedBuilds )
+			{
+				BaseCookingFlags |= ECookInitializationFlags::IterateSharedBuild | ECookInitializationFlags::IgnoreIniSettingsOutOfDate;
+			}
+
 			CookServer = NewObject<UCookOnTheFlyServer>();
 			CookServer->Initialize(ECookMode::CookOnTheFlyFromTheEditor, BaseCookingFlags);
 			CookServer->StartNetworkFileServer(false);
diff --git a/Engine/Source/Editor/UnrealEd/Public/Commandlets/AssetRegistryGenerator.h b/Engine/Source/Editor/UnrealEd/Public/Commandlets/AssetRegistryGenerator.h
--- a/Engine/Source/Editor/UnrealEd/Public/Commandlets/AssetRegistryGenerator.h
+++ b/Engine/Source/Editor/UnrealEd/Public/Commandlets/AssetRegistryGenerator.h
@@ -113,7 +113,7 @@ public:
 	/**
 	 * Saves generated asset registry data for each platform.
 	 */
-	bool SaveAssetRegistry(const FString& SandboxPath);
+	bool SaveAssetRegistry(const FString& SandboxPath, bool bSerializeDevelopmentAssetRegistry = true);
 
 	/** 
 	 * Writes out CookerOpenOrder.log file 
diff --git a/Engine/Source/Programs/AutomationTool/BuildGraph/Script.cs b/Engine/Source/Programs/AutomationTool/BuildGraph/Script.cs
--- a/Engine/Source/Programs/AutomationTool/BuildGraph/Script.cs
+++ b/Engine/Source/Programs/AutomationTool/BuildGraph/Script.cs
@@ -798,7 +798,6 @@ namespace AutomationTool
 				foreach (Node ReferencedNode in ResolveReferences(Element, RequiredNames))
 				{
 					NewBadge.Nodes.Add(ReferencedNode);
-					NewBadge.Nodes.UnionWith(ReferencedNode.OrderDependencies);
 				}
 				Graph.Badges.Add(NewBadge);
 			}
diff --git a/Engine/Source/Programs/AutomationTool/Scripts/AutomationScripts.Automation.csproj b/Engine/Source/Programs/AutomationTool/Scripts/AutomationScripts.Automation.csproj
--- a/Engine/Source/Programs/AutomationTool/Scripts/AutomationScripts.Automation.csproj
+++ b/Engine/Source/Programs/AutomationTool/Scripts/AutomationScripts.Automation.csproj
@@ -80,6 +80,7 @@
     <Compile Include="CookCommand.Automation.cs" />
     <Compile Include="CookTargetCommand.Automation.cs" />
     <Compile Include="CopyBuildToStagingDirectory.Automation.cs" />
+    <Compile Include="CopySharedCookedBuild.Automation.cs" />
     <Compile Include="CopyUAT.Automation.cs" />
     <Compile Include="CopyUsingDistillFileSet.Automation.cs" />
     <Compile Include="DeployCommand.Automation.cs" />
@@ -100,6 +101,7 @@
     <Compile Include="PackageCommand.Automation.cs" />
     <Compile Include="Properties\AssemblyInfo.cs" />
     <Compile Include="RunProjectCommand.Automation.cs" />
+    <Compile Include="SharedCookedBuild.cs" />
     <Compile Include="StageTargetCommand.Automation.cs" />
     <Compile Include="StreamCopyDescription.Automation.cs" />
     <Compile Include="SubmitUtililzationReportToEC.Automation.cs" />
diff --git a/Engine/Source/Programs/AutomationTool/Scripts/CookCommand.Automation.cs b/Engine/Source/Programs/AutomationTool/Scripts/CookCommand.Automation.cs
--- a/Engine/Source/Programs/AutomationTool/Scripts/CookCommand.Automation.cs
+++ b/Engine/Source/Programs/AutomationTool/Scripts/CookCommand.Automation.cs
@@ -44,154 +44,6 @@ public partial class Project : CommandUtils
         return PluginPath;
     }
 
-	static Task CopySharedCookedBuildTask = null;
-
-	static bool CopySharedCookedBuildForTargetInternal( string CookedBuildPath, string CookPlatform, string LocalPath)
-	{
-
-		string BuildRoot = P4Enabled ? P4Env.BuildRootP4.Replace("/", "+") : "";
-		int RecentCL = P4Enabled ? P4Env.Changelist : 0;
-
-		BuildVersion Version;
-		if (BuildVersion.TryRead(out Version))
-		{
-			RecentCL = Version.Changelist;
-			BuildRoot = Version.BranchName;
-		}
-		System.GC.Collect();
-
-		// check to see if we have already synced this build ;)
-		var SyncedBuildFile = CombinePaths(LocalPath, "SyncedBuild.txt");
-		string BuildCL = "Invalid";
-		if (File.Exists(SyncedBuildFile))
-		{
-			BuildCL = File.ReadAllText(SyncedBuildFile);
-		}
-
-		if (RecentCL == 0 && CookedBuildPath.Contains("[CL]"))
-		{
-			Log("Unable to copy shared cooked build: Unable to determine CL number from P4 or UGS, and is required by SharedCookedBuildPath");
-			return false;
-		}
-
-		if (RecentCL == 0 && CookedBuildPath.Contains("[BRANCHNAME]"))
-		{
-			Log("Unable to copy shared cooked build: Unable to determine BRANCHNAME number from P4 or UGS, and is required by SharedCookedBuildPath");
-			return false;
-		}
-
-
-		CookedBuildPath = CookedBuildPath.Replace("[CL]", RecentCL.ToString());
-		CookedBuildPath = CookedBuildPath.Replace("[BRANCHNAME]", BuildRoot);
-		CookedBuildPath = CookedBuildPath.Replace("[PLATFORM]", CookPlatform);
-
-		if (Directory.Exists(CookedBuildPath) == false)
-		{
-			Log("Unable to copy shared cooked build: Unable to find shared build at location {0} check SharedCookedBuildPath in Engine.ini SharedCookedBuildSettings is correct", CookedBuildPath);
-			return false;
-		}
-
-		Log("Attempting download of latest shared build CL {0} from location {1}", RecentCL, CookedBuildPath);
-
-		if (BuildCL == RecentCL.ToString())
-		{
-			Log("Already downloaded latest shared build at CL {0}", RecentCL);
-			return false;
-		}
-		// delete all the stuff
-		Log("Deleting previous shared build because it was out of date");
-		CommandUtils.DeleteDirectory(LocalPath);
-		Directory.CreateDirectory(LocalPath);
-
-
-
-		string CookedBuildCookedDirectory = Path.Combine(CookedBuildPath, "Cooked");
-		CookedBuildCookedDirectory = Path.GetFullPath(CookedBuildCookedDirectory);
-		string LocalBuildCookedDirectory = Path.Combine(LocalPath, "Cooked");
-		LocalBuildCookedDirectory = Path.GetFullPath(LocalBuildCookedDirectory);
-		if (Directory.Exists(CookedBuildCookedDirectory))
-		{
-			foreach (string FileName in Directory.EnumerateFiles(CookedBuildCookedDirectory, "*.*", SearchOption.AllDirectories))
-			{
-				string SourceFileName = Path.GetFullPath(FileName);
-				string DestFileName = SourceFileName.Replace(CookedBuildCookedDirectory, LocalBuildCookedDirectory);
-				Directory.CreateDirectory(Path.GetDirectoryName(DestFileName));
-				File.Copy(SourceFileName, DestFileName);
-			}
-		}
-
-
-		CopySharedCookedBuildTask = Task.Run(() =>
-		{
-			// find all the files in the staged directory
-			string CookedBuildStagedDirectory = Path.GetFullPath(Path.Combine(CookedBuildPath, "Staged"));
-			string LocalBuildStagedDirectory = Path.GetFullPath(Path.Combine(LocalPath, "Staged"));
-			if (Directory.Exists(CookedBuildStagedDirectory))
-			{
-				foreach (string FileName in Directory.EnumerateFiles(CookedBuildStagedDirectory, "*.*", SearchOption.AllDirectories))
-				{
-					string SourceFileName = Path.GetFullPath(FileName);
-					string DestFileName = SourceFileName.Replace(CookedBuildStagedDirectory, LocalBuildStagedDirectory);
-					Directory.CreateDirectory(Path.GetDirectoryName(DestFileName));
-					File.Copy(SourceFileName, DestFileName);
-				}
-			}
-		}
-		);
-
-		File.WriteAllText(SyncedBuildFile, RecentCL.ToString());
-		return true;
-	}
-
-	static void CopySharedCookedBuildForTarget(ProjectParams Params,TargetPlatformDescriptor TargetPlatform, string CookPlatform)
-	{
-
-		string ProjectPath = Params.RawProjectPath.FullName;
-		var LocalPath = CombinePaths(GetDirectoryName(ProjectPath), "Saved", "SharedIterativeBuild", CookPlatform);
-
-		// get network location 
-		ConfigHierarchy Hierarchy = ConfigCache.ReadHierarchy(ConfigHierarchyType.Engine, DirectoryReference.FromFile(Params.RawProjectPath), TargetPlatform.Type);
-		List<string> CookedBuildPaths;
-		if (Hierarchy.GetArray("SharedCookedBuildSettings", "SharedCookedBuildPath", out CookedBuildPaths) == false)
-		{
-			Log("Unable to copy shared cooked build: SharedCookedBuildPath not set in Engine.ini SharedCookedBuildSettings");
-			return;
-		}
-		foreach ( var CookedBuildPath in CookedBuildPaths )
-		{
-			if (CopySharedCookedBuildForTargetInternal(CookedBuildPath, CookPlatform, LocalPath) == true)
-			{
-				return;
-			}
-		}
-
-		return;
-	}
-
-	static void CopySharedCookedBuild(ProjectParams Params)
-	{
-
-		if (!Params.NoClient)
-		{ 
-			foreach (var ClientPlatform in Params.ClientTargetPlatforms)
-			{
-				// Use the data platform, sometimes we will copy another platform's data
-				var DataPlatformDesc = Params.GetCookedDataPlatformForClientTarget(ClientPlatform);
-				string PlatformToCook = Platform.Platforms[DataPlatformDesc].GetCookPlatform(false, Params.Client);
-				CopySharedCookedBuildForTarget(Params, ClientPlatform, PlatformToCook);
-			}
-		}
-		if (Params.DedicatedServer)
-		{
-			foreach (var ServerPlatform in Params.ServerTargetPlatforms)
-			{
-				// Use the data platform, sometimes we will copy another platform's data
-				var DataPlatformDesc = Params.GetCookedDataPlatformForServerTarget(ServerPlatform);
-				string PlatformToCook = Platform.Platforms[DataPlatformDesc].GetCookPlatform(true, false);
-				CopySharedCookedBuildForTarget(Params, ServerPlatform, PlatformToCook);
-			}
-		}
-	}
 
     public static void Cook(ProjectParams Params)
 	{
@@ -343,7 +195,7 @@ public partial class Project : CommandUtils
                 }
 				if ( Params.IterateSharedCookedBuild)
 				{
-					CopySharedCookedBuild(Params);
+					SharedCookedBuild.CopySharedCookedBuild(Params);
 					CommandletParams += " -iteratesharedcookedbuild";					
 				}
 
@@ -432,10 +284,7 @@ public partial class Project : CommandUtils
 
                 CookCommandlet(Params.RawProjectPath, Params.UE4Exe, Maps, Dirs, InternationalizationPreset, CulturesToCook, CombineCommandletParams(PlatformsToCook.ToArray()), CommandletParams);
 
-				if (CopySharedCookedBuildTask != null)
-				{
-					CopySharedCookedBuildTask.Wait();
-				}
+				SharedCookedBuild.WaitForCopy();
             }
 			catch (Exception Ex)
 			{
diff --git a/Engine/Source/Programs/AutomationTool/Scripts/SharedCookedBuild.cs b/Engine/Source/Programs/AutomationTool/Scripts/SharedCookedBuild.cs
new file mode 100644
--- /dev/null
+++ b/Engine/Source/Programs/AutomationTool/Scripts/SharedCookedBuild.cs
@@ -0,0 +1,195 @@
+﻿// Copyright 1998-2017 Epic Games, Inc. All Rights Reserved.
+using System;
+using System.Collections.Generic;
+using System.IO;
+using System.Threading;
+using System.Reflection;
+using System.Linq;
+using System.Threading.Tasks;
+using AutomationTool;
+using UnrealBuildTool;
+
+
+public class SharedCookedBuild
+{
+	private static Task CopySharedCookedBuildTask = null;
+
+	private static bool CopySharedCookedBuildForTargetInternal(string CookedBuildPath, string CookPlatform, string LocalPath, bool bOnlyCopyAssetRegistry)
+	{
+
+		string BuildRoot = CommandUtils.P4Enabled ? CommandUtils.P4Env.BuildRootP4.Replace("/", "+") : "";
+		string RecentCL = CommandUtils.P4Enabled ? CommandUtils.P4Env.Changelist.ToString() : "";
+
+		BuildVersion Version;
+		if (BuildVersion.TryRead(out Version))
+		{
+			RecentCL = Version.Changelist.ToString();
+			BuildRoot = Version.BranchName;
+		}
+		System.GC.Collect();
+
+		// check to see if we have already synced this build ;)
+		var SyncedBuildFile = CommandUtils.CombinePaths(LocalPath, "SyncedBuild.txt");
+		string BuildCL = "Invalid";
+		if (File.Exists(SyncedBuildFile))
+		{
+			BuildCL = File.ReadAllText(SyncedBuildFile);
+		}
+
+		if (RecentCL == "" && CookedBuildPath.Contains("[CL]"))
+		{
+			CommandUtils.Log("Unable to copy shared cooked build: Unable to determine CL number from P4 or UGS, and is required by SharedCookedBuildPath");
+			return false;
+		}
+
+		if (RecentCL == "" && CookedBuildPath.Contains("[BRANCHNAME]"))
+		{
+			CommandUtils.Log("Unable to copy shared cooked build: Unable to determine BRANCHNAME number from P4 or UGS, and is required by SharedCookedBuildPath");
+			return false;
+		}
+
+
+		CookedBuildPath = CookedBuildPath.Replace("[CL]", RecentCL.ToString());
+		CookedBuildPath = CookedBuildPath.Replace("[BRANCHNAME]", BuildRoot);
+		CookedBuildPath = CookedBuildPath.Replace("[PLATFORM]", CookPlatform);
+
+		if (Directory.Exists(CookedBuildPath) == false)
+		{
+			CommandUtils.Log("Unable to copy shared cooked build: Unable to find shared build at location {0} check SharedCookedBuildPath in Engine.ini SharedCookedBuildSettings is correct", CookedBuildPath);
+			return false;
+		}
+
+		CommandUtils.Log("Attempting download of latest shared build CL {0} from location {1}", RecentCL, CookedBuildPath);
+
+		if (BuildCL == RecentCL.ToString())
+		{
+			CommandUtils.Log("Already downloaded latest shared build at CL {0}", RecentCL);
+			return false;
+		}
+
+		if (bOnlyCopyAssetRegistry)
+		{
+			RecentCL += "RegistryOnly";
+		}
+
+		if ( BuildCL == RecentCL )
+		{
+			CommandUtils.Log("Already downloaded latest shared build at CL {0}", RecentCL);
+			return false;
+		}
+
+		// delete all the stuff
+		CommandUtils.Log("Deleting previous shared build because it was out of date");
+		CommandUtils.DeleteDirectory(LocalPath);
+		Directory.CreateDirectory(LocalPath);
+
+
+
+		string CookedBuildCookedDirectory = Path.Combine(CookedBuildPath, "Cooked");
+		CookedBuildCookedDirectory = Path.GetFullPath(CookedBuildCookedDirectory);
+		string LocalBuildCookedDirectory = Path.Combine(LocalPath, "Cooked");
+		LocalBuildCookedDirectory = Path.GetFullPath(LocalBuildCookedDirectory);
+		if (Directory.Exists(CookedBuildCookedDirectory))
+		{
+			foreach (string FileName in Directory.EnumerateFiles(CookedBuildCookedDirectory, "*.*", SearchOption.AllDirectories))
+			{
+				string SourceFileName = Path.GetFullPath(FileName);
+				string DestFileName = SourceFileName.Replace(CookedBuildCookedDirectory, LocalBuildCookedDirectory);
+				Directory.CreateDirectory(Path.GetDirectoryName(DestFileName));
+				File.Copy(SourceFileName, DestFileName);
+			}
+		}
+
+		if ( CopySharedCookedBuildTask != null )
+		{
+			WaitForCopy();
+		}
+
+		if (bOnlyCopyAssetRegistry == false)
+		{
+			CopySharedCookedBuildTask = Task.Run(() =>
+				{
+					// find all the files in the staged directory
+					string CookedBuildStagedDirectory = Path.GetFullPath(Path.Combine(CookedBuildPath, "Staged"));
+					string LocalBuildStagedDirectory = Path.GetFullPath(Path.Combine(LocalPath, "Staged"));
+					if (Directory.Exists(CookedBuildStagedDirectory))
+					{
+						foreach (string FileName in Directory.EnumerateFiles(CookedBuildStagedDirectory, "*.*", SearchOption.AllDirectories))
+						{
+							string SourceFileName = Path.GetFullPath(FileName);
+							string DestFileName = SourceFileName.Replace(CookedBuildStagedDirectory, LocalBuildStagedDirectory);
+							Directory.CreateDirectory(Path.GetDirectoryName(DestFileName));
+							File.Copy(SourceFileName, DestFileName);
+						}
+					}
+					File.WriteAllText(SyncedBuildFile, RecentCL.ToString());
+				}
+				);
+		}
+		else
+		{
+			File.WriteAllText(SyncedBuildFile, RecentCL.ToString());
+		}
+
+		
+		return true;
+	}
+
+	public static void CopySharedCookedBuildForTarget(string ProjectFullPath, UnrealTargetPlatform TargetPlatform, string CookPlatform, bool bOnlyCopyAssetRegistry = false)
+	{
+		var LocalPath = CommandUtils.CombinePaths(Path.GetDirectoryName(ProjectFullPath), "Saved", "SharedIterativeBuild", CookPlatform);
+
+		FileReference ProjectFileRef = new FileReference(ProjectFullPath);
+		// get network location 
+		ConfigHierarchy Hierarchy = ConfigCache.ReadHierarchy(ConfigHierarchyType.Engine, DirectoryReference.FromFile(ProjectFileRef), TargetPlatform);
+		List<string> CookedBuildPaths;
+		if (Hierarchy.GetArray("SharedCookedBuildSettings", "SharedCookedBuildPath", out CookedBuildPaths) == false)
+		{
+			CommandUtils.Log("Unable to copy shared cooked build: SharedCookedBuildPath not set in Engine.ini SharedCookedBuildSettings");
+			return;
+		}
+		foreach (var CookedBuildPath in CookedBuildPaths)
+		{
+			if (CopySharedCookedBuildForTargetInternal(CookedBuildPath, CookPlatform, LocalPath, bOnlyCopyAssetRegistry) == true)
+			{
+				return;
+			}
+		}
+
+		return;
+	}
+
+	public static void CopySharedCookedBuild(ProjectParams Params)
+	{
+
+		if (!Params.NoClient)
+		{
+			foreach (var ClientPlatform in Params.ClientTargetPlatforms)
+			{
+				// Use the data platform, sometimes we will copy another platform's data
+				var DataPlatformDesc = Params.GetCookedDataPlatformForClientTarget(ClientPlatform);
+				string PlatformToCook = Platform.Platforms[DataPlatformDesc].GetCookPlatform(false, Params.Client);
+				CopySharedCookedBuildForTarget(Params.RawProjectPath.FullName, ClientPlatform.Type, PlatformToCook);
+			}
+		}
+		if (Params.DedicatedServer)
+		{
+			foreach (var ServerPlatform in Params.ServerTargetPlatforms)
+			{
+				// Use the data platform, sometimes we will copy another platform's data
+				var DataPlatformDesc = Params.GetCookedDataPlatformForServerTarget(ServerPlatform);
+				string PlatformToCook = Platform.Platforms[DataPlatformDesc].GetCookPlatform(true, false);
+				CopySharedCookedBuildForTarget(Params.RawProjectPath.FullName, ServerPlatform.Type, PlatformToCook);
+			}
+		}
+	}
+
+	public static void WaitForCopy()
+	{
+		if (CopySharedCookedBuildTask != null)
+		{
+			CopySharedCookedBuildTask.Wait();
+		}
+	}
+
+}
\ No newline at end of file
diff --git a/Engine/Source/Programs/UnrealBuildTool/Configuration/UEBuildTarget.cs b/Engine/Source/Programs/UnrealBuildTool/Configuration/UEBuildTarget.cs
--- a/Engine/Source/Programs/UnrealBuildTool/Configuration/UEBuildTarget.cs
+++ b/Engine/Source/Programs/UnrealBuildTool/Configuration/UEBuildTarget.cs
@@ -2525,7 +2525,7 @@ namespace UnrealBuildTool
 
 			if ((BuildConfiguration.bXGEExport && BuildConfiguration.bGenerateManifest) || (!BuildConfiguration.bGenerateManifest && !BuildConfiguration.bCleanProject && !ProjectFileGenerator.bGenerateProjectFiles))
 			{
-				HashSet<UEBuildModuleCPP> ModulesToGenerateHeadersFor = GatherDependencyModules(AppBinaries);
+				HashSet<UEBuildModuleCPP> ModulesToGenerateHeadersFor = GatherDependencyModules(OriginalBinaries);
 
 				if (OnlyModules.Count > 0)
 				{
diff --git a/Engine/Source/Programs/UnrealBuildTool/System/ConfigFile.cs b/Engine/Source/Programs/UnrealBuildTool/System/ConfigFile.cs
--- a/Engine/Source/Programs/UnrealBuildTool/System/ConfigFile.cs
+++ b/Engine/Source/Programs/UnrealBuildTool/System/ConfigFile.cs
@@ -177,7 +177,7 @@ namespace UnrealBuildTool
 							{
 								if(!TryAddConfigLine(CurrentSection, Line, StartIdx, EndIdx, DefaultAction))
 								{
-									Console.WriteLine("Couldn't parse '{0}'", Line);
+									Console.WriteLine("Couldn't parse '{0}' in {1} of {2}", Line, CurrentSection, Location.FullName);
 								}
 								break;
 							}
diff --git a/Engine/Source/Runtime/AIModule/Classes/EnvironmentQuery/EnvQueryTypes.h b/Engine/Source/Runtime/AIModule/Classes/EnvironmentQuery/EnvQueryTypes.h
--- a/Engine/Source/Runtime/AIModule/Classes/EnvironmentQuery/EnvQueryTypes.h
+++ b/Engine/Source/Runtime/AIModule/Classes/EnvironmentQuery/EnvQueryTypes.h
@@ -705,7 +705,7 @@ class AIMODULE_API UEnvQueryTypes : public UObject
 
 public:
 	/** special test value assigned to items skipped by condition check */
-	static float SkippedItemValue;
+	static const float SkippedItemValue;
 
 	/** special value used for executing query steps to prevent them from being time sliced */
 	static float UnlimitedStepTime;
diff --git a/Engine/Source/Runtime/AIModule/Classes/Tasks/AITask_MoveTo.h b/Engine/Source/Runtime/AIModule/Classes/Tasks/AITask_MoveTo.h
--- a/Engine/Source/Runtime/AIModule/Classes/Tasks/AITask_MoveTo.h
+++ b/Engine/Source/Runtime/AIModule/Classes/Tasks/AITask_MoveTo.h
@@ -36,6 +36,10 @@ public:
 	DEPRECATED(4.12, "This function is now depreacted, please use version with FAIMoveRequest parameter")
 	void SetUp(AAIController* Controller, FVector GoalLocation, AActor* GoalActor = nullptr, float AcceptanceRadius = -1.f, bool bUsePathfinding = true, EAIOptionFlag::Type StopOnOverlap = EAIOptionFlag::Default, EAIOptionFlag::Type AcceptPartialPath = EAIOptionFlag::Default);
 
+	/** Allows custom move request tweaking. Note that all MoveRequest need to 
+	 *	be performed before PerformMove is called. */
+	FAIMoveRequest& GetMoveRequestRef() { return MoveRequest; }
+
 protected:
 	UPROPERTY(BlueprintAssignable)
 	FGenericGameplayTaskDelegate OnRequestFailed;
diff --git a/Engine/Source/Runtime/AIModule/Private/EnvironmentQuery/EnvQueryInstance.cpp b/Engine/Source/Runtime/AIModule/Private/EnvironmentQuery/EnvQueryInstance.cpp
--- a/Engine/Source/Runtime/AIModule/Private/EnvironmentQuery/EnvQueryInstance.cpp
+++ b/Engine/Source/Runtime/AIModule/Private/EnvironmentQuery/EnvQueryInstance.cpp
@@ -15,6 +15,12 @@
 #include "EnvironmentQuery/Generators/EnvQueryGenerator_Composite.h"
 #include "AISystem.h"
 
+#if UE_BUILD_SHIPPING
+#define eqs_ensure ensure
+#else
+#define eqs_ensure ensureAlways
+#endif
+
 //----------------------------------------------------------------------//
 // FEnvQueryDebugData
 //----------------------------------------------------------------------//
@@ -543,7 +549,7 @@ void FEnvQueryInstance::FItemIterator::HandleFailedTestResult()
 void FEnvQueryInstance::FItemIterator::StoreTestResult()
 {
 	CheckItemPassed();
-	ensureAlways(FMath::IsNaN(ItemScore) == false);
+	eqs_ensure(FMath::IsNaN(ItemScore) == false);
 
 #if USE_EQS_DEBUGGER
 	Instance.NumProcessedItems++;
@@ -577,7 +583,7 @@ void FEnvQueryInstance::FItemIterator::StoreTestResult()
 		}
 		else if (CachedScoreOp == EEnvTestScoreOperator::AverageScore && !bForced)
 		{
-			ensureAlways(NumPassedForItem != 0);
+			eqs_ensure(NumPassedForItem != 0);
 			ItemScore /= NumPassedForItem;
 		}
 
@@ -893,3 +899,5 @@ FBox FEnvQueryInstance::GetBoundingBox() const
 
 	return BBox;
 }
+
+#undef eqs_ensure
\ No newline at end of file
diff --git a/Engine/Source/Runtime/AIModule/Private/EnvironmentQuery/EnvQueryManager.cpp b/Engine/Source/Runtime/AIModule/Private/EnvironmentQuery/EnvQueryManager.cpp
--- a/Engine/Source/Runtime/AIModule/Private/EnvironmentQuery/EnvQueryManager.cpp
+++ b/Engine/Source/Runtime/AIModule/Private/EnvironmentQuery/EnvQueryManager.cpp
@@ -25,6 +25,7 @@
 
 extern UNREALED_API UEditorEngine* GEditor;
 #endif // WITH_EDITOR
+#include "TimeGuard.h"
 
 DEFINE_LOG_CATEGORY(LogEQS);
 
@@ -316,6 +317,7 @@ bool UEnvQueryManager::AbortQuery(int32 RequestID)
 
 void UEnvQueryManager::Tick(float DeltaTime)
 {
+	SCOPE_TIME_GUARD_MS(TEXT("UEnvQueryManager::Tick"), 10);
 	SCOPE_CYCLE_COUNTER(STAT_AI_EQS_Tick);
 	SET_DWORD_STAT(STAT_AI_EQS_NumInstances, RunningQueries.Num());
 
diff --git a/Engine/Source/Runtime/AIModule/Private/EnvironmentQuery/EnvQueryTypes.cpp b/Engine/Source/Runtime/AIModule/Private/EnvironmentQuery/EnvQueryTypes.cpp
--- a/Engine/Source/Runtime/AIModule/Private/EnvironmentQuery/EnvQueryTypes.cpp
+++ b/Engine/Source/Runtime/AIModule/Private/EnvironmentQuery/EnvQueryTypes.cpp
@@ -20,7 +20,7 @@
 
 #define LOCTEXT_NAMESPACE "EnvQueryGenerator"
 
-float UEnvQueryTypes::SkippedItemValue = -FLT_MAX;
+const float UEnvQueryTypes::SkippedItemValue = -MAX_flt;
 float UEnvQueryTypes::UnlimitedStepTime = -1.f;
 
 //----------------------------------------------------------------------//
diff --git a/Engine/Source/Runtime/AIModule/Private/Navigation/PathFollowingComponent.cpp b/Engine/Source/Runtime/AIModule/Private/Navigation/PathFollowingComponent.cpp
--- a/Engine/Source/Runtime/AIModule/Private/Navigation/PathFollowingComponent.cpp
+++ b/Engine/Source/Runtime/AIModule/Private/Navigation/PathFollowingComponent.cpp
@@ -886,6 +886,11 @@ void UPathFollowingComponent::UpdatePathSegment()
 
 	FMetaNavMeshPath* MetaNavPath = bIsUsingMetaPath ? Path->CastPath<FMetaNavMeshPath>() : nullptr;
 
+	/** it's possible that finishing this move request will result in another request
+	 *	which won't be easily detectable from this function. This simple local
+	 *	variable gives us this knowledge. */
+	const FAIRequestID MoveRequestId = GetCurrentRequestId();
+
 	// if agent has control over its movement, check finish conditions
 	const FVector CurrentLocation = MovementComp->GetActorFeetLocation();
 	const bool bCanUpdateState = HasMovementAuthority();
@@ -933,7 +938,10 @@ void UPathFollowingComponent::UpdatePathSegment()
 		}
 	}
 
-	if (bCanUpdateState && Status == EPathFollowingStatus::Moving)
+	if (bCanUpdateState 
+		&& Status == EPathFollowingStatus::Moving
+		// still the same move request
+		&& MoveRequestId == GetCurrentRequestId())
 	{
 		// check waypoint switch condition in meta paths
 		if (MetaNavPath)
diff --git a/Engine/Source/Runtime/AnimGraphRuntime/Private/AnimNodes/AnimNode_LayeredBoneBlend.cpp b/Engine/Source/Runtime/AnimGraphRuntime/Private/AnimNodes/AnimNode_LayeredBoneBlend.cpp
--- a/Engine/Source/Runtime/AnimGraphRuntime/Private/AnimNodes/AnimNode_LayeredBoneBlend.cpp
+++ b/Engine/Source/Runtime/AnimGraphRuntime/Private/AnimNodes/AnimNode_LayeredBoneBlend.cpp
@@ -18,7 +18,7 @@ void FAnimNode_LayeredBoneBlend::Initialize(const FAnimationInitializeContext& C
 	// initialize children
 	BasePose.Initialize(Context);
 
-	if ( NumPoses > 0 )
+	if (NumPoses > 0)
 	{
 		for (int32 ChildIndex = 0; ChildIndex < NumPoses; ++ChildIndex)
 		{
@@ -41,7 +41,7 @@ void FAnimNode_LayeredBoneBlend::PostCompile(const class USkeleton* InSkeleton)
 
 void FAnimNode_LayeredBoneBlend::RebuildCacheData(const USkeleton* InSkeleton)
 {
-	// make sure it's inished post load yet, this might cause more things to cache in initialize instead of during cooking
+	// make sure it's Finished post load yet, this might cause more things to cache in initialize instead of during cooking
 	if (InSkeleton && !(InSkeleton->GetFlags() & RF_NeedPostLoad))
 	{
 		FAnimationRuntime::CreateMaskWeights(PerBoneBlendWeights, LayerSetup, InSkeleton);
@@ -64,11 +64,12 @@ void FAnimNode_LayeredBoneBlend::ReinitializeBoneBlendWeights(const FBoneContain
 	
 	// build desired bone weights
 	const TArray<FBoneIndexType>& RequiredBoneIndices = RequiredBones.GetBoneIndicesArray();
-	DesiredBoneBlendWeights.SetNumZeroed(RequiredBoneIndices.Num());
-	for (int32 RequiredBoneIndex = 0; RequiredBoneIndex < RequiredBoneIndices.Num(); ++RequiredBoneIndex)
+	const int32 NumRequiredBones = RequiredBoneIndices.Num();
+	DesiredBoneBlendWeights.SetNumZeroed(NumRequiredBones);
+	for (int32 RequiredBoneIndex=0; RequiredBoneIndex<NumRequiredBones; RequiredBoneIndex++)
 	{
-		int32 SkeletonIndex = RequiredBones.GetSkeletonIndex(FCompactPoseBoneIndex(RequiredBoneIndex));
-		DesiredBoneBlendWeights[RequiredBoneIndex] = PerBoneBlendWeights[SkeletonIndex];
+		const int32 SkeletonBoneIndex = RequiredBones.GetSkeletonIndex(FCompactPoseBoneIndex(RequiredBoneIndex));
+		DesiredBoneBlendWeights[RequiredBoneIndex] = PerBoneBlendWeights[SkeletonBoneIndex];
 	}
 	
 	CurrentBoneBlendWeights.Reset(DesiredBoneBlendWeights.Num());
@@ -136,7 +137,23 @@ void FAnimNode_LayeredBoneBlend::Update(const FAnimationUpdateContext& Context)
 		{
 			if (bHasRelevantPoses == false)
 			{
-				FAnimationRuntime::UpdateDesiredBoneWeight(DesiredBoneBlendWeights, CurrentBoneBlendWeights, BlendWeights);
+				// If our cache is invalid, attempt to update it.
+				if (IsCacheInvalid(Context.AnimInstanceProxy->GetSkeleton()))
+				{
+					ReinitializeBoneBlendWeights(Context.AnimInstanceProxy->GetRequiredBones(), Context.AnimInstanceProxy->GetSkeleton());
+
+					// If Cache is still invalid, we don't have correct DesiredBoneBlendWeights, so abort.
+					// bHasRelevantPoses == false, will passthrough in evaluate.
+					if (!ensure(IsCacheInvalid(Context.AnimInstanceProxy->GetSkeleton())))
+					{
+						break;
+					}
+				}
+				else
+				{
+					FAnimationRuntime::UpdateDesiredBoneWeight(DesiredBoneBlendWeights, CurrentBoneBlendWeights, BlendWeights);
+				}
+
 				bHasRelevantPoses = true;
 
 				if(bBlendRootMotionBasedOnRootBone)
@@ -192,7 +209,7 @@ void FAnimNode_LayeredBoneBlend::Evaluate(FPoseContext& Output)
 
 		for (int32 ChildIndex = 0; ChildIndex < NumPoses; ++ChildIndex)
 		{
-			if (BlendWeights[ChildIndex] > ZERO_ANIMWEIGHT_THRESH)
+			if (FAnimWeight::IsRelevant(BlendWeights[ChildIndex]))
 			{
 				FPoseContext CurrentPoseContext(Output);
 				BlendPoses[ChildIndex].Evaluate(CurrentPoseContext);
diff --git a/Engine/Source/Runtime/AnimGraphRuntime/Private/AnimNodes/AnimNode_RotationOffsetBlendSpace.cpp b/Engine/Source/Runtime/AnimGraphRuntime/Private/AnimNodes/AnimNode_RotationOffsetBlendSpace.cpp
--- a/Engine/Source/Runtime/AnimGraphRuntime/Private/AnimNodes/AnimNode_RotationOffsetBlendSpace.cpp
+++ b/Engine/Source/Runtime/AnimGraphRuntime/Private/AnimNodes/AnimNode_RotationOffsetBlendSpace.cpp
@@ -26,11 +26,13 @@ void FAnimNode_RotationOffsetBlendSpace::CacheBones(const FAnimationCacheBonesCo
 
 void FAnimNode_RotationOffsetBlendSpace::UpdateAssetPlayer(const FAnimationUpdateContext& Context)
 {
+	EvaluateGraphExposedInputs.Execute(Context);
+
 	ActualAlpha = AlphaScaleBias.ApplyTo(Alpha);
 	bIsLODEnabled = IsLODEnabled(Context.AnimInstanceProxy, LODThreshold);
 	if (bIsLODEnabled && FAnimWeight::IsRelevant(ActualAlpha))
 	{
-		FAnimNode_BlendSpacePlayer::UpdateAssetPlayer(Context);
+		UpdateInternal(Context);
 	}
 
 	BasePose.Update(Context);
@@ -59,7 +61,7 @@ void FAnimNode_RotationOffsetBlendSpace::GatherDebugData(FNodeDebugData& DebugDa
 {
 	FString DebugLine = DebugData.GetNodeName(this);
 	
-	DebugLine += FString::Printf(TEXT("(Play Time: %.3f)"), InternalTimeAccumulator);
+	DebugLine += FString::Printf(TEXT("Alpha (%.1f%%) PlayTime (%.3f)"), ActualAlpha * 100.f, InternalTimeAccumulator);
 	DebugData.AddDebugItem(DebugLine);
 	
 	BasePose.GatherDebugData(DebugData);
diff --git a/Engine/Source/Runtime/AnimGraphRuntime/Private/BoneControllers/AnimNode_ScaleChainLength.cpp b/Engine/Source/Runtime/AnimGraphRuntime/Private/BoneControllers/AnimNode_ScaleChainLength.cpp
--- a/Engine/Source/Runtime/AnimGraphRuntime/Private/BoneControllers/AnimNode_ScaleChainLength.cpp
+++ b/Engine/Source/Runtime/AnimGraphRuntime/Private/BoneControllers/AnimNode_ScaleChainLength.cpp
@@ -9,7 +9,9 @@
 // FAnimNode_ScaleChainLength
 
 FAnimNode_ScaleChainLength::FAnimNode_ScaleChainLength()
+	: Alpha(1.f)
 {
+	ChainInitialLength = EScaleChainInitialLength::FixedDefaultLengthValue;
 }
 
 void FAnimNode_ScaleChainLength::Initialize(const FAnimationInitializeContext& Context)
@@ -22,6 +24,8 @@ void FAnimNode_ScaleChainLength::Update(const FAnimationUpdateContext& Context)
 {
 	EvaluateGraphExposedInputs.Execute(Context);
 	InputPose.Update(Context);
+
+	ActualAlpha = AlphaScaleBias.ApplyTo(Alpha);
 }
 
 void FAnimNode_ScaleChainLength::CacheBones(const FAnimationCacheBonesContext& Context)
@@ -37,8 +41,7 @@ void FAnimNode_ScaleChainLength::Evaluate(FPoseContext& Output)
 	// Evaluate incoming pose into our output buffer.
 	InputPose.Evaluate(Output);
 
-	const float FinalAlpha = AlphaScaleBias.ApplyTo(Alpha);
-	if (!FAnimWeight::IsRelevant(FinalAlpha))
+	if (!FAnimWeight::IsRelevant(ActualAlpha))
 	{
 		return;
 	}
@@ -60,7 +63,7 @@ void FAnimNode_ScaleChainLength::Evaluate(FPoseContext& Output)
 
 		if (bBoneSetupIsValid)
 		{
-			FCompactPoseBoneIndex StartBoneIndex = ChainStartBone.GetCompactPoseIndex(BoneContainer);
+			const FCompactPoseBoneIndex StartBoneIndex = ChainStartBone.GetCompactPoseIndex(BoneContainer);
 			FCompactPoseBoneIndex BoneIndex = ChainEndBone.GetCompactPoseIndex(BoneContainer);
 			ChainBoneIndices.Add(BoneIndex);
 			if (BoneIndex != INDEX_NONE)
@@ -92,8 +95,9 @@ void FAnimNode_ScaleChainLength::Evaluate(FPoseContext& Output)
 	const FTransform StartTransformCompSpace = CSPose.GetComponentSpaceTransform(ChainBoneIndices[0]);
 
 	const float DesiredChainLength = (TargetLocationCompSpace - StartTransformCompSpace.GetLocation()).Size();
-	const float ChainLengthScale = !FMath::IsNearlyZero(DefaultChainLength) ? (DesiredChainLength / DefaultChainLength) : 1.f;
-	const float ChainLengthScaleWithAlpha = FMath::LerpStable(1.f, ChainLengthScale, FinalAlpha);
+	const float InitialChainLength = GetInitialChainLength(Output.Pose, CSPose);
+	const float ChainLengthScale = !FMath::IsNearlyZero(InitialChainLength) ? (DesiredChainLength / InitialChainLength) : 1.f;
+	const float ChainLengthScaleWithAlpha = FMath::LerpStable(1.f, ChainLengthScale, ActualAlpha);
 
 	// If we're not going to scale anything, early out.
 	if (FMath::IsNearlyEqual(ChainLengthScaleWithAlpha, 1.f))
@@ -101,19 +105,46 @@ void FAnimNode_ScaleChainLength::Evaluate(FPoseContext& Output)
 		return;
 	}
 
-	// Scale translation of all bones.
-	FCompactPose& BSPose = Output.Pose;
+	// Scale translation of all bones in local space.
+	FCompactPose& LSPose = Output.Pose;
 	for (const FCompactPoseBoneIndex& BoneIndex : ChainBoneIndices)
 	{
 		// Get bone space transform, scale transition.
-		BSPose[BoneIndex].ScaleTranslation(ChainLengthScaleWithAlpha);
+		LSPose[BoneIndex].ScaleTranslation(ChainLengthScaleWithAlpha);
 	}
 }
 
+float FAnimNode_ScaleChainLength::GetInitialChainLength(FCompactPose& InLSPose, FCSPose<FCompactPose>& InCSPose) const
+{
+	switch (ChainInitialLength)
+	{
+	case EScaleChainInitialLength::Distance : 
+	{
+		const FVector ChainStartLocation = InCSPose.GetComponentSpaceTransform(ChainBoneIndices[0]).GetLocation();
+		const FVector ChainEndLocation = InCSPose.GetComponentSpaceTransform(ChainBoneIndices.Last()).GetLocation();
+		return (ChainEndLocation - ChainStartLocation).Size();
+	}
+
+	case EScaleChainInitialLength::ChainLength :
+	{
+		float ChainLength = 0.f;
+		for (const FCompactPoseBoneIndex& BoneIndex : ChainBoneIndices)
+		{
+			ChainLength += InLSPose[BoneIndex].GetTranslation().Size();
+		}
+		return ChainLength;
+	}
+	};
+
+	// Fallback is using fixed value DefaultChainLength.
+	return DefaultChainLength;
+}
+
 void FAnimNode_ScaleChainLength::GatherDebugData(FNodeDebugData& DebugData)
 {
 	FString DebugLine = DebugData.GetNodeName(this);
-
+	DebugLine += FString::Printf(TEXT("Alpha (%.1f%%)"), ActualAlpha * 100.f);
 	DebugData.AddDebugItem(DebugLine);
+
 	InputPose.GatherDebugData(DebugData);
 }
diff --git a/Engine/Source/Runtime/AnimGraphRuntime/Public/BoneControllers/AnimNode_ScaleChainLength.h b/Engine/Source/Runtime/AnimGraphRuntime/Public/BoneControllers/AnimNode_ScaleChainLength.h
--- a/Engine/Source/Runtime/AnimGraphRuntime/Public/BoneControllers/AnimNode_ScaleChainLength.h
+++ b/Engine/Source/Runtime/AnimGraphRuntime/Public/BoneControllers/AnimNode_ScaleChainLength.h
@@ -12,8 +12,19 @@
 
 class USkeletalMeshComponent;
 
+UENUM(BlueprintType)
+enum class EScaleChainInitialLength : uint8
+{
+	/** Use the 'DefaultChainLength' input value. */
+	FixedDefaultLengthValue,
+	/** Use distance between 'ChainStartBone' and 'ChainEndBone' (in Component Space) */
+	Distance,
+	/* Use animated chain length (length in local space of every bone from 'ChainStartBone' to 'ChainEndBone' */
+	ChainLength,
+};
+
 /**
- *	
+ *	Scale the length of a chain of bones.
  */
 USTRUCT(BlueprintInternalUseOnly)
 struct ANIMGRAPHRUNTIME_API FAnimNode_ScaleChainLength : public FAnimNode_Base
@@ -33,12 +44,18 @@ struct ANIMGRAPHRUNTIME_API FAnimNode_ScaleChainLength : public FAnimNode_Base
 	UPROPERTY(EditAnywhere, Category = ScaleChainLength)
 	FBoneReference ChainEndBone;
 
+	UPROPERTY(EditAnywhere, Category = ScaleChainLength)
+	EScaleChainInitialLength ChainInitialLength;
+
 	UPROPERTY(EditAnywhere, BlueprintReadWrite, Category = ScaleChainLength, meta = (PinShownByDefault))
 	FVector TargetLocation;
 
 	UPROPERTY(EditAnywhere, BlueprintReadWrite, Category = Settings, meta = (PinShownByDefault))
 	mutable float Alpha;
 
+	UPROPERTY(Transient)
+	float ActualAlpha;
+
 	UPROPERTY(EditAnywhere, Category = Settings)
 	FInputScaleBias AlphaScaleBias;
 
@@ -57,4 +74,7 @@ public:
 	virtual void Evaluate(FPoseContext& Output) override;
 	virtual void GatherDebugData(FNodeDebugData& DebugData) override;
 	// End of FAnimNode_Base interface
+
+protected:
+	float GetInitialChainLength(FCompactPose& InLSPose, FCSPose<FCompactPose>& InCSPose) const;
 };
diff --git a/Engine/Source/Runtime/CookedIterativeFile/CookedIterativeFile.Build.cs b/Engine/Source/Runtime/CookedIterativeFile/CookedIterativeFile.Build.cs
new file mode 100644
--- /dev/null
+++ b/Engine/Source/Runtime/CookedIterativeFile/CookedIterativeFile.Build.cs
@@ -0,0 +1,17 @@
+// Copyright 1998-2017 Epic Games, Inc. All Rights Reserved.
+
+using UnrealBuildTool;
+
+public class CookedIterativeFile : ModuleRules
+{
+	public CookedIterativeFile(ReadOnlyTargetRules Target) : base(Target)
+	{
+        PublicDependencyModuleNames.AddRange(
+            new string[] {
+                "Core",
+                "NetworkFile",
+                "Sockets",
+            }
+        );
+    }
+}
diff --git a/Engine/Source/Runtime/CookedIterativeFile/Private/CookedIterativeNetworkFile.cpp b/Engine/Source/Runtime/CookedIterativeFile/Private/CookedIterativeNetworkFile.cpp
new file mode 100644
--- /dev/null
+++ b/Engine/Source/Runtime/CookedIterativeFile/Private/CookedIterativeNetworkFile.cpp
@@ -0,0 +1,730 @@
+// Copyright 1998-2017 Epic Games, Inc. All Rights Reserved.
+
+#include "CookedIterativeNetworkFile.h"
+#include "Templates/ScopedPointer.h"
+#include "Misc/CommandLine.h"
+#include "Misc/ScopeLock.h"
+#include "Modules/ModuleManager.h"
+#include "HAL/IPlatformFileModule.h"
+#include "UniquePtr.h"
+#include "Misc/EngineVersion.h"
+
+DEFINE_LOG_CATEGORY(LogCookedIterativeNetworkFile);
+
+bool FCookedIterativeNetworkFile::InitializeInternal(IPlatformFile* Inner, const TCHAR* HostIP)
+{
+	if ( Inner->GetLowerLevel() == nullptr )
+	{
+		UE_LOG(LogCookedIterativeNetworkFile, Fatal, TEXT("Platform file is missing it's inner.  Is pak file deployed?") );
+	}
+	
+
+	PakPlatformFile = Inner;
+
+	return FNetworkPlatformFile::InitializeInternal(Inner->GetLowerLevel(), HostIP );
+}
+
+
+void FCookedIterativeNetworkFile::ProcessServerCachedFilesResponse(FArrayReader& Response, const int32 ServerPackageVersion, const int32 ServerPackageLicenseeVersion)
+{
+	FNetworkPlatformFile::ProcessServerCachedFilesResponse(Response, ServerPackageVersion, ServerPackageLicenseeVersion);
+
+
+	// some of our stuff is on here too!
+	
+	// receive a list of the cache files and their timestamps
+	TMap<FString, FDateTime> ServerValidPakFileFiles;
+	Response << ServerValidPakFileFiles;
+
+
+	for ( const auto& ServerPakFileFile : ServerValidPakFileFiles )
+	{
+		// the server should contain all these files
+		FString Path = FPaths::GetPath( ServerPakFileFile.Key);
+		FString Filename = FPaths::GetCleanFilename(ServerPakFileFile.Key);
+
+		FServerTOC::FDirectory* ServerDirectory = ServerFiles.FindDirectory(Path);
+		if ( !ServerDirectory)
+		{
+			UE_LOG(LogCookedIterativeNetworkFile, Warning, TEXT("Unable to find directory %s while trying to resolve pak file %s"), *Path, *ServerPakFileFile.Key );
+		}
+		else
+		{
+			FDateTime* ServerDate = ServerDirectory->Find(Filename);
+			if ( !ServerDate )
+			{
+				UE_LOG(LogCookedIterativeNetworkFile, Warning, TEXT("Unable to find filename %s while trying to resolve pak file %s"), *Filename, *ServerPakFileFile.Key);
+			}
+		}
+
+		// check the file is accessable 
+		if (!PakPlatformFile->FileExists(*ServerPakFileFile.Key))
+		{
+			UE_LOG(LogCookedIterativeNetworkFile, Warning, TEXT("Unable to find file %s in pak file.  Server says it should be!"), *ServerPakFileFile.Key)
+		}
+
+		ValidPakFileFiles.AddFileOrDirectory(ServerPakFileFile.Key, ServerPakFileFile.Value);
+
+		UE_LOG(LogCookedIterativeNetworkFile, Display, TEXT("Using pak file %s"), *ServerPakFileFile.Key);
+	}
+}
+
+FString FCookedIterativeNetworkFile::GetVersionInfo() const
+{
+	FString VersionInfo = FString::Printf(TEXT("%s %d"), *FEngineVersion::CompatibleWith().GetBranch(), FEngineVersion::CompatibleWith().GetChangelist() );
+
+	return VersionInfo;
+}
+
+bool FCookedIterativeNetworkFile::ShouldPassToPak(const TCHAR* Filename) const
+{
+	if ( FCString::Stricmp( *FPaths::GetExtension(Filename), TEXT("ufont") ) == 0 )
+	{
+		FString Path = FPaths::GetPath( Filename );
+
+		const FServerTOC::FDirectory* Directory = ValidPakFileFiles.FindDirectory(Path);
+		if ( Directory )
+		{
+			FString StringFilename = Filename;
+			for ( const auto& File : *Directory )
+			{
+				if ( File.Key.StartsWith( StringFilename ) )
+				{
+					return true;
+				}
+			}
+			return false;
+		}
+		else
+		{
+			return false;
+		}
+	}
+
+	if ( ValidPakFileFiles.FindFile(Filename) != nullptr )
+	{
+		return true;
+	}
+	// if we are searching for the .uexp or .ubulk or any of those content extra files.... 
+	// then change it to the original,  if we are using the pak version of the original file then we want the bulk / uexp file to be the same
+	// potential async issue here if the file is invalidated after we start loading the original but before we load the uexp 
+	// not much we can do about that though... 
+	FString OriginalName = FPaths::ChangeExtension(Filename, TEXT("uasset"));
+	if ( ValidPakFileFiles.FindFile(OriginalName) )
+	{
+		return true;
+	}
+	OriginalName = FPaths::ChangeExtension(Filename, TEXT("umap"));
+	if ( ValidPakFileFiles.FindFile(OriginalName))
+	{
+		return true;
+	}
+
+	return false;
+}
+
+
+bool		FCookedIterativeNetworkFile::FileExists(const TCHAR* Filename)
+{
+	if ( ShouldPassToPak(Filename) )
+	{
+		PakPlatformFile->FileExists(Filename);
+		return true;
+	}
+	return FNetworkPlatformFile::FileExists(Filename);
+}
+
+int64 FCookedIterativeNetworkFile::FileSize(const TCHAR* Filename)
+{
+	if ( ShouldPassToPak(Filename) )
+	{
+		return PakPlatformFile->FileSize(Filename);
+	}
+	else
+	{
+		return FNetworkPlatformFile::FileSize(Filename);
+	}
+}
+bool FCookedIterativeNetworkFile::DeleteFile(const TCHAR* Filename)
+{
+	// delete both of these entries
+	ValidPakFileFiles.RemoveFileOrDirectory(Filename);
+	return FNetworkPlatformFile::DeleteFile(Filename);
+}
+
+bool FCookedIterativeNetworkFile::IsReadOnly(const TCHAR* Filename)
+{
+	if ( ShouldPassToPak(Filename) )
+	{
+		return PakPlatformFile->IsReadOnly(Filename);
+	}
+
+	return FNetworkPlatformFile::IsReadOnly(Filename);
+}
+
+bool FCookedIterativeNetworkFile::SetReadOnly(const TCHAR* Filename, bool bNewReadOnlyValue)
+{
+	if (ShouldPassToPak(Filename))
+	{
+		return PakPlatformFile->SetReadOnly(Filename, bNewReadOnlyValue);
+	}
+	return FNetworkPlatformFile::SetReadOnly(Filename, bNewReadOnlyValue);
+}
+FDateTime FCookedIterativeNetworkFile::GetTimeStamp(const TCHAR* Filename)
+{
+	if ( ShouldPassToPak(Filename) )
+	{
+		return PakPlatformFile->GetTimeStamp(Filename);
+	}
+	return FNetworkPlatformFile::GetTimeStamp(Filename);
+}
+
+
+void FCookedIterativeNetworkFile::OnFileUpdated(const FString& LocalFilename)
+{
+	FNetworkPlatformFile::OnFileUpdated(LocalFilename);
+	ValidPakFileFiles.RemoveFileOrDirectory(LocalFilename);
+}
+
+
+IFileHandle* FCookedIterativeNetworkFile::OpenRead(const TCHAR* Filename, bool bAllowWrite)
+{
+	if ( ShouldPassToPak(Filename) )
+	{
+		return PakPlatformFile->OpenRead(Filename, bAllowWrite);
+	}
+	return FNetworkPlatformFile::OpenRead(Filename, bAllowWrite);
+}
+IFileHandle* FCookedIterativeNetworkFile::OpenWrite(const TCHAR* Filename, bool bAppend, bool bAllowRead)
+{
+	if (ShouldPassToPak(Filename))
+	{
+		return PakPlatformFile->OpenWrite(Filename, bAppend, bAllowRead);
+	}
+	// FNetworkPlatformFile::CreateDirectoryTree(Directory);
+	return FNetworkPlatformFile::OpenWrite(Filename, bAppend, bAllowRead);
+}
+
+bool FCookedIterativeNetworkFile::CopyFile(const TCHAR* To, const TCHAR* From, EPlatformFileRead ReadFlags, EPlatformFileWrite WriteFlags)
+{
+	return PakPlatformFile->CopyFile(To, From, ReadFlags, WriteFlags);
+}
+
+bool FCookedIterativeNetworkFile::MoveFile(const TCHAR* To, const TCHAR* From)
+{
+	if (ShouldPassToPak(From))
+	{
+		return PakPlatformFile->MoveFile(To, From);
+	}
+	return FNetworkPlatformFile::MoveFile(To, From);
+}
+
+bool FCookedIterativeNetworkFile::DirectoryExists(const TCHAR* Directory)
+{
+	if ( ValidPakFileFiles.FindDirectory(Directory) )
+	{
+		return true;
+	}
+	return FNetworkPlatformFile::DirectoryExists(Directory);
+}
+
+bool FCookedIterativeNetworkFile::CreateDirectoryTree(const TCHAR* Directory)
+{
+	return FNetworkPlatformFile::CreateDirectoryTree(Directory);
+}
+
+bool FCookedIterativeNetworkFile::CreateDirectory(const TCHAR* Directory)
+{
+	return FNetworkPlatformFile::CreateDirectory(Directory);
+}
+
+bool FCookedIterativeNetworkFile::DeleteDirectory(const TCHAR* Directory)
+{
+	bool bSucceeded = false;
+	if( ValidPakFileFiles.FindDirectory(Directory) )
+	{
+		bSucceeded |= ValidPakFileFiles.RemoveFileOrDirectory(Directory) > 0 ? true : false;
+	}
+	bSucceeded |= FNetworkPlatformFile::DeleteDirectory(Directory);
+	return bSucceeded;
+}
+
+bool FCookedIterativeNetworkFile::DeleteDirectoryRecursively(const TCHAR* Directory)
+{
+	bool bSucceeded = false;
+	if (ValidPakFileFiles.FindDirectory(Directory))
+	{
+		bSucceeded |= ValidPakFileFiles.RemoveFileOrDirectory(Directory) > 0 ? true : false;
+	}
+	bSucceeded |= FNetworkPlatformFile::DeleteDirectoryRecursively(Directory);
+	return bSucceeded;
+}
+
+
+bool		FCookedIterativeNetworkFile::IterateDirectory(const TCHAR* Directory, IPlatformFile::FDirectoryVisitor& Visitor)
+{
+	//return PakPlatformFile->IterateDirectory(Directory, Visitor);
+	return FNetworkPlatformFile::IterateDirectory(Directory, Visitor);
+}
+
+bool FCookedIterativeNetworkFile::IterateDirectoryRecursively(const TCHAR* Directory, IPlatformFile::FDirectoryVisitor& Visitor)
+{
+	return FNetworkPlatformFile::IterateDirectoryRecursively(Directory, Visitor);
+}
+
+
+FFileStatData FCookedIterativeNetworkFile::GetStatData(const TCHAR* FilenameOrDirectory)
+{
+	if (ShouldPassToPak(FilenameOrDirectory))
+	{
+		return PakPlatformFile->GetStatData(FilenameOrDirectory);
+	}
+	return FNetworkPlatformFile::GetStatData(FilenameOrDirectory);
+}
+
+
+bool FCookedIterativeNetworkFile::IterateDirectoryStat(const TCHAR* Directory, IPlatformFile::FDirectoryStatVisitor& Visitor)
+{
+	return FNetworkPlatformFile::IterateDirectoryStat(Directory, Visitor);
+}
+
+bool FCookedIterativeNetworkFile::IterateDirectoryStatRecursively(const TCHAR* Directory, IPlatformFile::FDirectoryStatVisitor& Visitor)
+{
+	return FNetworkPlatformFile::IterateDirectoryStatRecursively(Directory, Visitor);
+}
+
+
+
+
+
+
+
+/**
+ * A helper class for wrapping some of the network file payload specifics
+ */
+/*
+class FStreamingNetworkFileArchive
+	: public FBufferArchive
+{
+public:
+
+	FStreamingNetworkFileArchive(uint32 Command)
+	{
+		// make sure the command is at the start
+		*this << Command;
+	}
+
+	// helper to serialize TCHAR* (there are a lot)
+	FORCEINLINE friend FStreamingNetworkFileArchive& operator<<(FStreamingNetworkFileArchive& Ar, const TCHAR*& Str)
+	{
+		FString Temp(Str);
+		Ar << Temp;
+		return Ar;
+	}
+};
+*/
+
+
+/*
+class FStreamingNetworkFileHandle
+	: public IFileHandle
+{
+	FStreamingNetworkPlatformFile& Network;
+	FString Filename;
+	uint64 HandleId;
+	int64 FilePos;
+	int64 FileSize;
+	bool bWritable;
+	bool bReadable;
+	uint8 BufferCache[2][GBufferCacheSize];
+	int64 CacheStart[2];
+	int64 CacheEnd[2];
+	bool LazySeek;
+	int32 CurrentCache;
+
+public:
+
+	FStreamingNetworkFileHandle(FStreamingNetworkPlatformFile& InNetwork, const TCHAR* InFilename, uint64 InHandleId, int64 InFileSize, bool bWriting)
+		: Network(InNetwork)
+		, Filename(InFilename)
+		, HandleId(InHandleId)
+		, FilePos(0)
+		, FileSize(InFileSize)
+		, bWritable(bWriting)
+		, bReadable(!bWriting)
+		, LazySeek(false)
+		,CurrentCache(0)
+	{
+		CacheStart[0] = CacheStart[1] = -1;
+		CacheEnd[0] = CacheEnd[1] = -1;
+	}
+
+	~FStreamingNetworkFileHandle()
+	{
+		Network.SendCloseMessage(HandleId);
+	}
+
+	virtual int64 Size() override
+	{
+		return FileSize;
+	}
+
+	virtual int64 Tell() override
+	{
+		return FilePos;
+	}
+
+	virtual bool Seek(int64 NewPosition) override
+	{
+		if( bWritable )
+		{
+			if( NewPosition == FilePos )
+			{
+				return true;
+			}
+
+			if (NewPosition >= 0 && NewPosition <= FileSize)
+			{
+				if (Network.SendSeekMessage(HandleId, NewPosition))
+				{
+					FilePos = NewPosition;
+
+					return true;
+				}
+			}
+		}
+		else if( bReadable )
+		{
+			if (NewPosition >= 0 && NewPosition <= FileSize)
+			{
+				if (NewPosition < CacheStart[0] || NewPosition >= CacheEnd[0] || CacheStart[0] == -1)
+				{
+					if (NewPosition < CacheStart[1] || NewPosition >= CacheEnd[1] || CacheStart[1] == -1)
+					{
+						LazySeek = true;
+						FilePos = NewPosition;
+						if (CacheStart[CurrentCache] != -1)
+						{
+							CurrentCache++;
+							CurrentCache %= 2;
+							CacheStart[CurrentCache] = -1; // Invalidate the cache
+						}
+
+						return true;
+					}
+					else
+					{
+						LazySeek = false;
+						FilePos = NewPosition;
+						CurrentCache = 1;
+
+						return true;
+					}
+				}
+				else
+				{
+					LazySeek = false;
+					FilePos = NewPosition;
+					CurrentCache = 0;
+
+					return true;
+				}
+			}
+		}
+
+		return false;
+	}
+
+	virtual bool SeekFromEnd(int64 NewPositionRelativeToEnd = 0) override
+	{
+		return Seek(FileSize + NewPositionRelativeToEnd);
+	}
+
+	virtual bool Read(uint8* Destination, int64 BytesToRead) override
+	{
+		bool Result = false;
+		if (bReadable && BytesToRead >= 0 && BytesToRead + FilePos <= FileSize)
+		{
+			if (BytesToRead == 0)
+			{
+				Result = true;
+			}
+			else
+			{
+				if (BytesToRead > GBufferCacheSize) // reading more than we cache
+				{
+					// if the file position is within the cache, copy out the remainder of the cache
+					if (CacheStart[CurrentCache] != -1 && FilePos >= CacheStart[CurrentCache] && FilePos < CacheEnd[CurrentCache])
+					{
+						int64 CopyBytes = CacheEnd[CurrentCache]-FilePos;
+						FMemory::Memcpy(Destination, BufferCache[CurrentCache]+(FilePos-CacheStart[CurrentCache]), CopyBytes);
+						FilePos += CopyBytes;
+						BytesToRead -= CopyBytes;
+						Destination += CopyBytes;
+					}
+
+					if (Network.SendSeekMessage(HandleId, FilePos))
+					{
+						Result = Network.SendReadMessage(HandleId, Destination, BytesToRead);
+					}
+					if (Result)
+					{
+						FilePos += BytesToRead;
+						CurrentCache++;
+						CurrentCache %= 2;
+						CacheStart[CurrentCache] = -1; // Invalidate the cache
+					}
+				}
+				else
+				{
+					Result = true;
+
+					// need to update the cache
+					if (CacheStart[CurrentCache] == -1 && FileSize < GBufferCacheSize)
+					{
+						Result = Network.SendReadMessage(HandleId, BufferCache[CurrentCache], FileSize);
+						if (Result)
+						{
+							CacheStart[CurrentCache] = 0;
+							CacheEnd[CurrentCache] = FileSize;
+						}
+					}
+					else if (FilePos + BytesToRead > CacheEnd[CurrentCache] || CacheStart[CurrentCache] == -1 || FilePos < CacheStart[CurrentCache])
+					{
+						// copy the data from FilePos to the end of the Cache to the destination as long as it is in the cache
+						if (CacheStart[CurrentCache] != -1 && FilePos >= CacheStart[CurrentCache] && FilePos < CacheEnd[CurrentCache])
+						{
+							int64 CopyBytes = CacheEnd[CurrentCache]-FilePos;
+							FMemory::Memcpy(Destination, BufferCache[CurrentCache]+(FilePos-CacheStart[CurrentCache]), CopyBytes);
+							FilePos += CopyBytes;
+							BytesToRead -= CopyBytes;
+							Destination += CopyBytes;
+						}
+
+						// switch to the other cache
+						if (CacheStart[CurrentCache] != -1)
+						{
+							CurrentCache++;
+							CurrentCache %= 2;
+						}
+
+						int64 SizeToRead = GBufferCacheSize;
+
+						if (FilePos + SizeToRead > FileSize)
+						{
+							SizeToRead = FileSize-FilePos;
+						}
+
+						if (Network.SendSeekMessage(HandleId, FilePos))
+						{
+							Result = Network.SendReadMessage(HandleId, BufferCache[CurrentCache], SizeToRead);
+						}
+
+						if (Result)
+						{
+							CacheStart[CurrentCache] = FilePos;
+							CacheEnd[CurrentCache] = FilePos + SizeToRead;
+						}
+					}
+
+					// copy from the cache to the destination
+					if (Result)
+					{
+						FMemory::Memcpy(Destination, BufferCache[CurrentCache]+(FilePos-CacheStart[CurrentCache]), BytesToRead);
+						FilePos += BytesToRead;
+					}
+				}
+			}
+		}
+
+		return Result;
+	}
+
+	virtual bool Write(const uint8* Source, int64 BytesToWrite) override
+	{
+		bool Result = false;
+
+		if (bWritable && BytesToWrite >= 0)
+		{
+			if (BytesToWrite == 0)
+			{
+				Result = true;
+			}
+			else
+			{
+				Result = Network.SendWriteMessage(HandleId, Source, BytesToWrite);
+
+				if (Result)
+				{
+					FilePos += BytesToWrite;
+					FileSize = FMath::Max<int64>(FilePos, FileSize);
+				}
+			}
+		}
+
+		return Result;
+	}
+};
+*/
+
+
+bool FCookedIterativeNetworkFile::ShouldBeUsed(IPlatformFile* Inner, const TCHAR* CmdLine) const
+{
+	bool bResult = FNetworkPlatformFile::ShouldBeUsed(Inner, CmdLine);
+
+	if (bResult)
+	{
+		bResult = FParse::Param(CmdLine, TEXT("precookednetwork"));
+	}
+
+	return bResult;
+}
+
+/*
+bool FCookedIterativeNetworkFile::InitializeInternal(IPlatformFile* Inner, const TCHAR* HostIP)
+{
+	// look for the commandline that will read files from over the network
+	if (HostIP == nullptr)
+	{
+		UE_LOG(LogStreamingPlatformFile, Error, TEXT("No Host IP specified in the commandline."));
+		bIsUsable = false;
+
+		return false;
+	}
+
+	// optionally get the port from the command line
+	int32 OverridePort;
+	if (FParse::Value(FCommandLine::Get(), TEXT("fileserverport="), OverridePort))
+	{
+		UE_LOG(LogStreamingPlatformFile, Display, TEXT("Overriding file server port: %d"), OverridePort);
+		FileServerPort = OverridePort;
+	}
+
+	// Send the filenames and timestamps to the server.
+	FNetworkFileArchive Payload(NFS_Messages::GetFileList);
+	FillGetFileList(Payload, true);
+
+	// Send the directories over, and wait for a response.
+	FArrayReader Response;
+
+	if(SendPayloadAndReceiveResponse(Payload,Response))
+	{
+		// Receive the cooked version information.
+		int32 ServerPackageVersion = 0;
+		int32 ServerPackageLicenseeVersion = 0;
+		ProcessServerInitialResponse(Response, ServerPackageVersion, ServerPackageLicenseeVersion);
+
+		// Make sure we can sync a file.
+		FString TestSyncFile = FPaths::Combine(*(FPaths::EngineDir()), TEXT("Config/BaseEngine.ini"));
+		IFileHandle* TestFileHandle = OpenRead(*TestSyncFile);
+		if (TestFileHandle != nullptr)
+		{
+			uint8* FileContents = (uint8*)FMemory::Malloc(TestFileHandle->Size());
+			if (!TestFileHandle->Read(FileContents, TestFileHandle->Size()))
+			{
+				UE_LOG(LogStreamingPlatformFile, Fatal, TEXT("Could not read test file %s."), *TestSyncFile);
+			}
+			FMemory::Free(FileContents);
+			delete TestFileHandle;
+		}
+		else
+		{
+			UE_LOG(LogStreamingPlatformFile, Fatal, TEXT("Could not open test file %s."), *TestSyncFile);
+		}
+
+		FCommandLine::AddToSubprocessCommandline( *FString::Printf( TEXT("-StreamingHostIP=%s"), HostIP ) );
+
+		return true; 
+	}
+
+	return false; 
+}*/
+
+
+FCookedIterativeNetworkFile::~FCookedIterativeNetworkFile()
+{
+}
+
+/*
+bool FStreamingNetworkPlatformFile::IterateDirectory(const TCHAR* InDirectory, IPlatformFile::FDirectoryVisitor& Visitor)
+{
+	FString RelativeDirectory = InDirectory;
+	MakeStandardNetworkFilename(RelativeDirectory);
+	// for .dll, etc searches that don't specify a path, we need to strip off the path
+	// before we send it to the visitor
+	bool bHadNoPath = InDirectory[0] == 0;
+
+	// we loop until this is false
+	bool RetVal = true;
+
+	// Find the directory in TOC
+	FServerTOC::FDirectory* ServerDirectory = ServerFiles.FindDirectory(RelativeDirectory);
+	if (ServerDirectory != nullptr)
+	{
+		// loop over the server files and look if they are in this exact directory
+		for (FServerTOC::FDirectory::TIterator It(*ServerDirectory); It && RetVal == true; ++It)
+		{
+			if (FPaths::GetPath(It.Key()) == RelativeDirectory)
+			{
+				// timestamps of 0 mean directories
+				bool bIsDirectory = It.Value() == 0;
+			
+				// visit (stripping off the path if needed)
+				RetVal = Visitor.Visit(bHadNoPath ? *FPaths::GetCleanFilename(It.Key()) : *It.Key(), bIsDirectory);
+			}
+		}
+	}
+
+	return RetVal;
+}
+
+
+bool FStreamingNetworkPlatformFile::IterateDirectoryRecursively(const TCHAR* InDirectory, IPlatformFile::FDirectoryVisitor& Visitor)
+{
+	FString RelativeDirectory = InDirectory;
+	MakeStandardNetworkFilename(RelativeDirectory);
+	// we loop until this is false
+	bool RetVal = true;
+
+	// loop over the server TOC
+	for (TMap<FString, FServerTOC::FDirectory*>::TIterator DirIt(ServerFiles.Directories); DirIt && RetVal == true; ++DirIt)
+	{
+		if (DirIt.Key().StartsWith(RelativeDirectory))
+		{
+			FServerTOC::FDirectory& ServerDirectory = *DirIt.Value();
+		
+			// loop over the server files and look if they are in this exact directory
+			for (FServerTOC::FDirectory::TIterator It(ServerDirectory); It && RetVal == true; ++It)
+			{
+				// timestamps of 0 mean directories
+				bool bIsDirectory = It.Value() == 0;
+
+				// visit!
+				RetVal = Visitor.Visit(*It.Key(), bIsDirectory);
+			}
+		}
+	}
+
+	return RetVal;
+}
+
+*/
+
+
+/**
+ * Module for the streaming file
+ */
+class FCookedIterativeFileModule
+	: public IPlatformFileModule
+{
+public:
+
+	virtual IPlatformFile* GetPlatformFile() override
+	{
+		static TUniquePtr<IPlatformFile> AutoDestroySingleton = MakeUnique<FCookedIterativeNetworkFile>();
+		return AutoDestroySingleton.Get();
+	}
+};
+
+
+IMPLEMENT_MODULE(FCookedIterativeFileModule, CookedIterativeFile);
diff --git a/Engine/Source/Runtime/CookedIterativeFile/Public/CookedIterativeNetworkFile.h b/Engine/Source/Runtime/CookedIterativeFile/Public/CookedIterativeNetworkFile.h
new file mode 100644
--- /dev/null
+++ b/Engine/Source/Runtime/CookedIterativeFile/Public/CookedIterativeNetworkFile.h
@@ -0,0 +1,196 @@
+// Copyright 1998-2017 Epic Games, Inc. All Rights Reserved.
+
+#pragma once
+
+#include "CoreMinimal.h"
+#include "GenericPlatform/GenericPlatformFile.h"
+#include "Misc/Paths.h"
+#include "NetworkMessage.h"
+#include "NetworkPlatformFile.h"
+
+DECLARE_LOG_CATEGORY_EXTERN(LogCookedIterativeNetworkFile, Display, All);
+
+
+#if 0
+/**
+ * Visitor to gather local files with their timestamps
+ */
+class STREAMINGFILE_API FStreamingLocalTimestampVisitor : public IPlatformFile::FDirectoryVisitor
+{
+private:
+
+	/** The file interface to use for any file operations */
+	IPlatformFile& FileInterface;
+
+	/** true if we want directories in this list */
+	bool bCacheDirectories;
+
+	/** A list of directories that we should not traverse into */
+	TArray<FString> DirectoriesToIgnore;
+
+	/** A list of directories that we should only go one level into */
+	TArray<FString> DirectoriesToNotRecurse;
+
+public:
+
+	/** Relative paths to local files and their timestamps */
+	TMap<FString, FDateTime> FileTimes;
+	
+	FStreamingLocalTimestampVisitor(IPlatformFile& InFileInterface, const TArray<FString>& InDirectoriesToIgnore, const TArray<FString>& InDirectoriesToNotRecurse, bool bInCacheDirectories=false)
+		: FileInterface(InFileInterface)
+		, bCacheDirectories(bInCacheDirectories)
+	{
+		// make sure the paths are standardized, since the Visitor will assume they are standard
+		for (int32 DirIndex = 0; DirIndex < InDirectoriesToIgnore.Num(); DirIndex++)
+		{
+			FString DirToIgnore = InDirectoriesToIgnore[DirIndex];
+			FPaths::MakeStandardFilename(DirToIgnore);
+			DirectoriesToIgnore.Add(DirToIgnore);
+		}
+
+		for (int32 DirIndex = 0; DirIndex < InDirectoriesToNotRecurse.Num(); DirIndex++)
+		{
+			FString DirToNotRecurse = InDirectoriesToNotRecurse[DirIndex];
+			FPaths::MakeStandardFilename(DirToNotRecurse);
+			DirectoriesToNotRecurse.Add(DirToNotRecurse);
+		}
+	}
+
+	virtual bool Visit(const TCHAR* FilenameOrDirectory, bool bIsDirectory)
+	{
+		// make sure all paths are "standardized" so the other end can match up with it's own standardized paths
+		FString RelativeFilename = FilenameOrDirectory;
+		FPaths::MakeStandardFilename(RelativeFilename);
+
+		// cache files and optionally directories
+		if (!bIsDirectory)
+		{
+			FileTimes.Add(RelativeFilename, FileInterface.GetTimeStamp(FilenameOrDirectory));
+		}
+		else if (bCacheDirectories)
+		{
+			// we use a timestamp of 0 to indicate a directory
+			FileTimes.Add(RelativeFilename, 0);
+		}
+
+		// iterate over directories we care about
+		if (bIsDirectory)
+		{
+			bool bShouldRecurse = true;
+			// look in all the ignore directories looking for a match
+			for (int32 DirIndex = 0; DirIndex < DirectoriesToIgnore.Num() && bShouldRecurse; DirIndex++)
+			{
+				if (RelativeFilename.StartsWith(DirectoriesToIgnore[DirIndex]))
+				{
+					bShouldRecurse = false;
+				}
+			}
+
+			if (bShouldRecurse == true)
+			{
+				// If it is a directory that we should not recurse (ie we don't want to process subdirectories of it)
+				// handle that case as well...
+				for (int32 DirIndex = 0; DirIndex < DirectoriesToNotRecurse.Num() && bShouldRecurse; DirIndex++)
+				{
+					if (RelativeFilename.StartsWith(DirectoriesToNotRecurse[DirIndex]))
+					{
+						// Are we more than level deep in that directory?
+						FString CheckFilename = RelativeFilename.Right(RelativeFilename.Len() - DirectoriesToNotRecurse[DirIndex].Len());
+						if (CheckFilename.Len() > 1)
+						{
+							bShouldRecurse = false;
+						}
+					}
+				}
+			}
+
+			// recurse if we should
+			if (bShouldRecurse)
+			{
+				FileInterface.IterateDirectory(FilenameOrDirectory, *this);
+			}
+		}
+
+		return true;
+	}
+};
+#endif
+
+
+/**
+ * Wrapper to redirect the low level file system to a server
+ */
+class COOKEDITERATIVEFILE_API FCookedIterativeNetworkFile
+	: public FNetworkPlatformFile
+{
+	friend class FAsyncFileSync;
+
+public:
+
+	/** Default Constructor */
+	FCookedIterativeNetworkFile() 
+	{ 
+		ConnectionFlags = EConnectionFlags::PreCookedIterative;
+		HeartbeatFrequency = 1; // increase the heartbeat frequency
+	}
+
+	/** Virtual destructor */
+	virtual ~FCookedIterativeNetworkFile();
+
+public:
+
+	static const TCHAR* GetTypeName()
+	{
+		return TEXT("CookedIterativeNetworkFile");
+	}
+
+protected:
+
+	virtual bool InitializeInternal(IPlatformFile* Inner, const TCHAR* HostIP) override;
+
+	// need to override what FNetworkPlatformFile does here
+	virtual void ProcessServerCachedFilesResponse(FArrayReader& InReponse, const int32 ServerPackageVersion, const int32 ServerPackageLicenseeVersion) override;
+
+	// IPlatformFile interface
+
+	virtual bool ShouldBeUsed(IPlatformFile* Inner, const TCHAR* CmdLine) const override;
+
+	virtual bool		FileExists(const TCHAR* Filename) override;
+	virtual int64		FileSize(const TCHAR* Filename) override;
+	virtual bool		DeleteFile(const TCHAR* Filename) override;
+	virtual bool		IsReadOnly(const TCHAR* Filename) override;
+	virtual bool		MoveFile(const TCHAR* To, const TCHAR* From) override;
+	virtual bool		SetReadOnly(const TCHAR* Filename, bool bNewReadOnlyValue) override;
+	virtual FDateTime	GetTimeStamp(const TCHAR* Filename) override;
+
+	virtual IFileHandle*	OpenRead(const TCHAR* Filename, bool bAllowWrite = false) override;
+	virtual IFileHandle*	OpenWrite(const TCHAR* Filename, bool bAppend = false, bool bAllowRead = false) override;
+	virtual bool		DirectoryExists(const TCHAR* Directory) override;
+	virtual bool		CreateDirectoryTree(const TCHAR* Directory) override;
+	virtual bool		CreateDirectory(const TCHAR* Directory) override;
+	virtual bool		DeleteDirectory(const TCHAR* Directory) override;
+
+	virtual FFileStatData GetStatData(const TCHAR* FilenameOrDirectory) override;
+
+	virtual bool		IterateDirectory(const TCHAR* Directory, IPlatformFile::FDirectoryVisitor& Visitor) override;
+	virtual bool		IterateDirectoryRecursively(const TCHAR* Directory, IPlatformFile::FDirectoryVisitor& Visitor) override;
+
+	virtual bool		IterateDirectoryStat(const TCHAR* Directory, IPlatformFile::FDirectoryStatVisitor& Visitor) override;
+	virtual bool		IterateDirectoryStatRecursively(const TCHAR* Directory, IPlatformFile::FDirectoryStatVisitor& Visitor) override;
+
+	virtual bool		DeleteDirectoryRecursively(const TCHAR* Directory) override;
+	virtual bool		CopyFile(const TCHAR* To, const TCHAR* From, EPlatformFileRead ReadFlags = EPlatformFileRead::None, EPlatformFileWrite WriteFlags = EPlatformFileWrite::None) override;
+
+
+	virtual void		OnFileUpdated(const FString& LocalFilename) override;
+
+	virtual FString		GetVersionInfo() const override;
+private:
+	bool ShouldPassToPak(const TCHAR* Filename) const;
+
+private:
+	FServerTOC ValidPakFileFiles;
+
+	IPlatformFile* PakPlatformFile;
+
+};
diff --git a/Engine/Source/Runtime/Core/Private/HAL/MallocBinned2.cpp b/Engine/Source/Runtime/Core/Private/HAL/MallocBinned2.cpp
--- a/Engine/Source/Runtime/Core/Private/HAL/MallocBinned2.cpp
+++ b/Engine/Source/Runtime/Core/Private/HAL/MallocBinned2.cpp
@@ -11,6 +11,7 @@
 #include "GenericPlatform/GenericPlatformProcess.h"
 #include "Stats/Stats.h"
 #include "HAL/IConsoleManager.h"
+#include "HAL/MemoryMisc.h"
 
 
 #if BINNED2_ALLOW_RUNTIME_TWEAKING
@@ -52,6 +53,14 @@ static FAutoConsoleVariableRef GMallocBinned2AllocExtraCVar(
 
 #endif
 
+#if BINNED2_ALLOCATOR_STATS
+int64 AllocatedSmallPoolMemory = 0; // memory that's requested to be allocated by the game
+int64 AllocatedOSSmallPoolMemory = 0;
+
+int64 AllocatedLargePoolMemory = 0; // memory requests to the OS which don't fit in the small pool
+int64 AllocatedLargePoolMemoryWAlignment = 0; // when we allocate at OS level we need to align to a size
+#endif
+
 // Block sizes are based around getting the maximum amount of allocations per pool, with as little alignment waste as possible.
 // Block sizes should be close to even divisors of the system page size, and well distributed.
 // They must be 16-byte aligned as well.
@@ -480,6 +489,9 @@ struct FMallocBinned2::Private
 					// Free the OS memory.
 					NodePool->Unlink();
 					Allocator.CachedOSPageAllocator.Free(BasePtrOfNode, Allocator.PageSize);
+#if BINNED2_ALLOCATOR_STATS
+					AllocatedOSSmallPoolMemory -= ((int64)Allocator.PageSize);
+#endif
 				}
 
 				Node = NextNode;
@@ -525,6 +537,9 @@ FMallocBinned2::FPoolInfo& FMallocBinned2::FPoolList::PushNewPoolToFront(FMalloc
 	{
 		Private::OutOfMemory(LocalPageSize);
 	}
+#if BINNED2_ALLOCATOR_STATS
+	AllocatedOSSmallPoolMemory += (int64)LocalPageSize;
+#endif
 	check(IsAligned(Free, LocalPageSize));
 	// Create pool
 	FPoolInfo* Result = Private::GetOrCreatePoolInfo(Allocator, Free, FPoolInfo::ECanary::FirstFreeBlockIsPtr, false);
@@ -623,6 +638,10 @@ void* FMallocBinned2::MallocExternal(SIZE_T Size, uint32 Alignment)
 	{
 		uint32 PoolIndex = BoundSizeToPoolIndex(Size);
 
+#if BINNED2_ALLOCATOR_STATS
+		FPlatformAtomics::InterlockedAdd(&AllocatedSmallPoolMemory, (int64)PoolIndexToBlockSize(PoolIndex));
+#endif
+
 		FPerThreadFreeBlockLists* Lists = GMallocBinned2PerThreadCaches ? FPerThreadFreeBlockLists::Get() : nullptr;
 		if (Lists)
 		{
@@ -689,6 +708,11 @@ void* FMallocBinned2::MallocExternal(SIZE_T Size, uint32 Alignment)
 	}
 	check(IsAligned(Result, PageSize) && IsOSAllocation(Result));
 
+#if BINNED2_ALLOCATOR_STATS
+	AllocatedLargePoolMemory += (int64)Size;
+	AllocatedLargePoolMemoryWAlignment += AlignedSize;
+#endif
+
 	// Create pool.
 	FPoolInfo* Pool = Private::GetOrCreatePoolInfo(*this, Result, FPoolInfo::ECanary::FirstFreeBlockIsOSAllocSize, false);
 	check(Size > 0 && Size <= AlignedSize && AlignedSize >= OsAllocationGranularity);
@@ -758,6 +782,11 @@ void* FMallocBinned2::ReallocExternal(void* Ptr, SIZE_T NewSize, uint32 Alignmen
 		return Result;
 	}
 
+#if BINNED2_ALLOCATOR_STATS
+	AllocatedLargePoolMemory += ((int64)NewSize) - ((int64)Pool->GetOSRequestedBytes());
+	// don't need to change the AllocatedLargePoolMemoryWAlignment because we didn't reallocate so it's the same size
+#endif
+	
 	Pool->SetOSAllocationSizes(NewSize, PoolOsBytes);
 
 	return Ptr;
@@ -772,6 +801,10 @@ void FMallocBinned2::FreeExternal(void* Ptr)
 		BasePtr->CanaryTest();
 		uint32 BlockSize = BasePtr->BlockSize;
 		uint32 PoolIndex = BasePtr->PoolIndex;
+
+#if BINNED2_ALLOCATOR_STATS
+		FPlatformAtomics::InterlockedAdd(&AllocatedSmallPoolMemory, -((int64)(BlockSize)));
+#endif
 		FBundleNode* BundlesToRecycle = nullptr;
 		FPerThreadFreeBlockLists* Lists = GMallocBinned2PerThreadCaches ? FPerThreadFreeBlockLists::Get() : nullptr;
 		if (Lists)
@@ -802,6 +835,12 @@ void FMallocBinned2::FreeExternal(void* Ptr)
 		}
 		UPTRINT PoolOsBytes = Pool->GetOsAllocatedBytes();
 		uint32 PoolOSRequestedBytes = Pool->GetOSRequestedBytes();
+
+#if BINNED2_ALLOCATOR_STATS
+		AllocatedLargePoolMemory -= ((int64)PoolOSRequestedBytes);
+		AllocatedLargePoolMemoryWAlignment -= ((int64)PoolOsBytes);
+#endif
+
 		checkf(PoolOSRequestedBytes <= PoolOsBytes, TEXT("FMallocBinned2::FreeExternal %d %d"), int32(PoolOSRequestedBytes), int32(PoolOsBytes));
 		Pool->SetCanary(FPoolInfo::ECanary::Unassigned, true, false);
 		// Free an OS allocation.
@@ -1027,6 +1066,39 @@ void FMallocBinned2::FFreeBlock::CanaryFail() const
 	UE_LOG(LogMemory, Fatal, TEXT("FMallocBinned2 Attempt to realloc an unrecognized block %p   canary == 0x%x != 0x%x"), (void*)this, (int32)Canary, (int32)FMallocBinned2::FFreeBlock::CANARY_VALUE);
 }
 
+void FMallocBinned2::GetAllocatorStats( FGenericMemoryStats& OutStats )
+{
+	OutStats.Add(TEXT("AllocatedSmallPoolMemory"), AllocatedSmallPoolMemory );
+	OutStats.Add(TEXT("AllocatedOSSmallPoolMemory"), AllocatedOSSmallPoolMemory);
+	OutStats.Add(TEXT("AllocatedLargePoolMemory"), AllocatedLargePoolMemory);
+	OutStats.Add(TEXT("AllocatedLargePoolMemoryWAlignment"), AllocatedLargePoolMemoryWAlignment);
+	OutStats.Add(TEXT("PageAllocatorFreeCacheSize"), CachedOSPageAllocator.GetCachedFreeTotal());
+
+	uint64 TotalAllocated = AllocatedSmallPoolMemory + AllocatedLargePoolMemory;
+	uint64 TotalOSAllocated = AllocatedOSSmallPoolMemory + AllocatedLargePoolMemoryWAlignment + CachedOSPageAllocator.GetCachedFreeTotal();
+
+	OutStats.Add(TEXT("TotalAllocated"), TotalAllocated);
+	OutStats.Add(TEXT("TotalOSAllocated"), TotalOSAllocated);
+}
+
+void FMallocBinned2::DumpAllocatorStats(class FOutputDevice& Ar)
+{
+#if BINNED2_ALLOCATOR_STATS
+	Ar.Logf(TEXT("FMallocBinned2 Mem report"));
+	Ar.Logf(TEXT("Small Pool"));
+	Ar.Logf(TEXT("Requested Allocations: %fmb  (including block size padding)"), ((double)AllocatedSmallPoolMemory) / (1024.0f * 1024.0f) );
+	Ar.Logf(TEXT("OS Allocated: %fmb"), ((double)AllocatedOSSmallPoolMemory) / (1024.0f * 1024.0f));
+	Ar.Logf(TEXT("Large Pool"));
+	Ar.Logf(TEXT("Requested Allocations: %fmb"), ((double)AllocatedLargePoolMemory) / (1024.0f * 1024.0f));
+	Ar.Logf(TEXT("OS Allocated: %fmb"), ((double)AllocatedLargePoolMemoryWAlignment) / (1024.0f * 1024.0f));
+
+	uint32 OSPageAllocatorCachedFreeSize = CachedOSPageAllocator.GetCachedFreeTotal();
+	Ar.Logf(TEXT("OS Page Allocator"));
+	Ar.Logf(TEXT("Cached free pages: %fmb"), ((double)OSPageAllocatorCachedFreeSize) / (1024.0f * 1024.0f));
+#else
+	Ar.Logf(TEXT("Allocator Stats for binned2 are not in this build set BINNED2_ALLOCATOR_STATS 1 in MallocBinned2.cpp"));
+#endif
+}
 #if !BINNED2_INLINE
 	#if PLATFORM_USES_FIXED_GMalloc_CLASS && !FORCE_ANSI_ALLOCATOR && USE_MALLOC_BINNED2
 		//#define FMEMORY_INLINE_FUNCTION_DECORATOR  FORCEINLINE
diff --git a/Engine/Source/Runtime/Core/Private/Misc/EngineVersion.cpp b/Engine/Source/Runtime/Core/Private/Misc/EngineVersion.cpp
--- a/Engine/Source/Runtime/Core/Private/Misc/EngineVersion.cpp
+++ b/Engine/Source/Runtime/Core/Private/Misc/EngineVersion.cpp
@@ -206,6 +206,11 @@ const FEngineVersion& FEngineVersion::CompatibleWith()
 	return CompatibleWithVersion;
 }
 
+const FString& FEngineVersion::GetBranchDescriptor() const
+{
+	return Branch;
+}
+
 bool FEngineVersion::OverrideCurrentVersionChangelist(int32 NewChangelist, int32 NewCompatibleChangelist)
 {
 	if(CurrentVersion.GetChangelist() != 0 || CompatibleWithVersion.GetChangelist() != 0)
diff --git a/Engine/Source/Runtime/Core/Private/Misc/TimeGuard.cpp b/Engine/Source/Runtime/Core/Private/Misc/TimeGuard.cpp
--- a/Engine/Source/Runtime/Core/Private/Misc/TimeGuard.cpp
+++ b/Engine/Source/Runtime/Core/Private/Misc/TimeGuard.cpp
@@ -9,7 +9,6 @@ DEFINE_LOG_CATEGORY_STATIC(LogTimeGuard, Log, All);
 
 TMap<const TCHAR*, FLightweightTimeGuard::FGuardInfo>  FLightweightTimeGuard::HitchData;
 bool FLightweightTimeGuard::bEnabled;
-double FLightweightTimeGuard::LastHitchTime;
 float FLightweightTimeGuard::FrameTimeThresholdMS = 1000.0 / 30.0;
 FCriticalSection FLightweightTimeGuard::ReportMutex;
 
@@ -27,8 +26,6 @@ void FLightweightTimeGuard::ClearData()
 {
 	FScopeLock Lock(&ReportMutex);
 	HitchData.Empty();
-	// don't capture any hitches immediately
-	LastHitchTime = FPlatformTime::Seconds();
 }
 
 void FLightweightTimeGuard::GetData(TMap<const TCHAR*, FGuardInfo>& Dest)
@@ -39,16 +36,6 @@ void FLightweightTimeGuard::GetData(TMap<const TCHAR*, FGuardInfo>& Dest)
 
 void FLightweightTimeGuard::ReportHitch(const TCHAR* InName, const float TimeMS)
 {
-	const double kHitchDebounceTime = 2.0;
-
-	// Don't report hitches that occur in the same 5sec window. This will also stop
-	// the outer scope of nested checks reporting
-	const double TimeNow = FPlatformTime::Seconds();
-	if (TimeNow < LastHitchTime + kHitchDebounceTime)
-	{
-		return;
-	}
-
 	FScopeLock lock(&ReportMutex);
 
 	FGuardInfo& Data = HitchData.FindOrAdd(InName);
@@ -64,8 +51,6 @@ void FLightweightTimeGuard::ReportHitch(const TCHAR* InName, const float TimeMS)
 	Data.Max = FMath::Max(Data.Max, TimeMS);
 	Data.LastTime = FDateTime::UtcNow();
 
-	LastHitchTime = TimeNow;
-
 	UE_LOG(LogTimeGuard, Warning, TEXT("Detected Hitch of %0.2fms in %s"), TimeMS, InName);
 }
 
diff --git a/Engine/Source/Runtime/Core/Public/HAL/Allocators/CachedOSPageAllocator.h b/Engine/Source/Runtime/Core/Public/HAL/Allocators/CachedOSPageAllocator.h
--- a/Engine/Source/Runtime/Core/Public/HAL/Allocators/CachedOSPageAllocator.h
+++ b/Engine/Source/Runtime/Core/Public/HAL/Allocators/CachedOSPageAllocator.h
@@ -47,6 +47,11 @@ struct TCachedOSPageAllocator : private FCachedOSPageAllocator
 		return FreeAllImpl(FreedPageBlocks, FreedPageBlocksNum, CachedTotal);
 	}
 
+	uint32 GetCachedFreeTotal()
+	{
+		return CachedTotal;
+	}
+
 private:
 	FFreePageBlock FreedPageBlocks[NumCacheBlocks];
 	uint32         FreedPageBlocksNum;
