Removing a bunch of commented code that was causing one of the doctool regexes to hang indefinitely.

diff --git a/Engine/Source/Runtime/CookedIterativeFile/Private/CookedIterativeNetworkFile.cpp b/Engine/Source/Runtime/CookedIterativeFile/Private/CookedIterativeNetworkFile.cpp
--- a/Engine/Source/Runtime/CookedIterativeFile/Private/CookedIterativeNetworkFile.cpp
+++ b/Engine/Source/Runtime/CookedIterativeFile/Private/CookedIterativeNetworkFile.cpp
@@ -293,281 +293,6 @@ bool FCookedIterativeNetworkFile::IterateDirectoryStatRecursively(const TCHAR* D
 	return FNetworkPlatformFile::IterateDirectoryStatRecursively(Directory, Visitor);
 }
 
-
-
-
-
-
-
-/**
- * A helper class for wrapping some of the network file payload specifics
- */
-/*
-class FStreamingNetworkFileArchive
-	: public FBufferArchive
-{
-public:
-
-	FStreamingNetworkFileArchive(uint32 Command)
-	{
-		// make sure the command is at the start
-		*this << Command;
-	}
-
-	// helper to serialize TCHAR* (there are a lot)
-	FORCEINLINE friend FStreamingNetworkFileArchive& operator<<(FStreamingNetworkFileArchive& Ar, const TCHAR*& Str)
-	{
-		FString Temp(Str);
-		Ar << Temp;
-		return Ar;
-	}
-};
-*/
-
-
-/*
-class FStreamingNetworkFileHandle
-	: public IFileHandle
-{
-	FStreamingNetworkPlatformFile& Network;
-	FString Filename;
-	uint64 HandleId;
-	int64 FilePos;
-	int64 FileSize;
-	bool bWritable;
-	bool bReadable;
-	uint8 BufferCache[2][GBufferCacheSize];
-	int64 CacheStart[2];
-	int64 CacheEnd[2];
-	bool LazySeek;
-	int32 CurrentCache;
-
-public:
-
-	FStreamingNetworkFileHandle(FStreamingNetworkPlatformFile& InNetwork, const TCHAR* InFilename, uint64 InHandleId, int64 InFileSize, bool bWriting)
-		: Network(InNetwork)
-		, Filename(InFilename)
-		, HandleId(InHandleId)
-		, FilePos(0)
-		, FileSize(InFileSize)
-		, bWritable(bWriting)
-		, bReadable(!bWriting)
-		, LazySeek(false)
-		,CurrentCache(0)
-	{
-		CacheStart[0] = CacheStart[1] = -1;
-		CacheEnd[0] = CacheEnd[1] = -1;
-	}
-
-	~FStreamingNetworkFileHandle()
-	{
-		Network.SendCloseMessage(HandleId);
-	}
-
-	virtual int64 Size() override
-	{
-		return FileSize;
-	}
-
-	virtual int64 Tell() override
-	{
-		return FilePos;
-	}
-
-	virtual bool Seek(int64 NewPosition) override
-	{
-		if( bWritable )
-		{
-			if( NewPosition == FilePos )
-			{
-				return true;
-			}
-
-			if (NewPosition >= 0 && NewPosition <= FileSize)
-			{
-				if (Network.SendSeekMessage(HandleId, NewPosition))
-				{
-					FilePos = NewPosition;
-
-					return true;
-				}
-			}
-		}
-		else if( bReadable )
-		{
-			if (NewPosition >= 0 && NewPosition <= FileSize)
-			{
-				if (NewPosition < CacheStart[0] || NewPosition >= CacheEnd[0] || CacheStart[0] == -1)
-				{
-					if (NewPosition < CacheStart[1] || NewPosition >= CacheEnd[1] || CacheStart[1] == -1)
-					{
-						LazySeek = true;
-						FilePos = NewPosition;
-						if (CacheStart[CurrentCache] != -1)
-						{
-							CurrentCache++;
-							CurrentCache %= 2;
-							CacheStart[CurrentCache] = -1; // Invalidate the cache
-						}
-
-						return true;
-					}
-					else
-					{
-						LazySeek = false;
-						FilePos = NewPosition;
-						CurrentCache = 1;
-
-						return true;
-					}
-				}
-				else
-				{
-					LazySeek = false;
-					FilePos = NewPosition;
-					CurrentCache = 0;
-
-					return true;
-				}
-			}
-		}
-
-		return false;
-	}
-
-	virtual bool SeekFromEnd(int64 NewPositionRelativeToEnd = 0) override
-	{
-		return Seek(FileSize + NewPositionRelativeToEnd);
-	}
-
-	virtual bool Read(uint8* Destination, int64 BytesToRead) override
-	{
-		bool Result = false;
-		if (bReadable && BytesToRead >= 0 && BytesToRead + FilePos <= FileSize)
-		{
-			if (BytesToRead == 0)
-			{
-				Result = true;
-			}
-			else
-			{
-				if (BytesToRead > GBufferCacheSize) // reading more than we cache
-				{
-					// if the file position is within the cache, copy out the remainder of the cache
-					if (CacheStart[CurrentCache] != -1 && FilePos >= CacheStart[CurrentCache] && FilePos < CacheEnd[CurrentCache])
-					{
-						int64 CopyBytes = CacheEnd[CurrentCache]-FilePos;
-						FMemory::Memcpy(Destination, BufferCache[CurrentCache]+(FilePos-CacheStart[CurrentCache]), CopyBytes);
-						FilePos += CopyBytes;
-						BytesToRead -= CopyBytes;
-						Destination += CopyBytes;
-					}
-
-					if (Network.SendSeekMessage(HandleId, FilePos))
-					{
-						Result = Network.SendReadMessage(HandleId, Destination, BytesToRead);
-					}
-					if (Result)
-					{
-						FilePos += BytesToRead;
-						CurrentCache++;
-						CurrentCache %= 2;
-						CacheStart[CurrentCache] = -1; // Invalidate the cache
-					}
-				}
-				else
-				{
-					Result = true;
-
-					// need to update the cache
-					if (CacheStart[CurrentCache] == -1 && FileSize < GBufferCacheSize)
-					{
-						Result = Network.SendReadMessage(HandleId, BufferCache[CurrentCache], FileSize);
-						if (Result)
-						{
-							CacheStart[CurrentCache] = 0;
-							CacheEnd[CurrentCache] = FileSize;
-						}
-					}
-					else if (FilePos + BytesToRead > CacheEnd[CurrentCache] || CacheStart[CurrentCache] == -1 || FilePos < CacheStart[CurrentCache])
-					{
-						// copy the data from FilePos to the end of the Cache to the destination as long as it is in the cache
-						if (CacheStart[CurrentCache] != -1 && FilePos >= CacheStart[CurrentCache] && FilePos < CacheEnd[CurrentCache])
-						{
-							int64 CopyBytes = CacheEnd[CurrentCache]-FilePos;
-							FMemory::Memcpy(Destination, BufferCache[CurrentCache]+(FilePos-CacheStart[CurrentCache]), CopyBytes);
-							FilePos += CopyBytes;
-							BytesToRead -= CopyBytes;
-							Destination += CopyBytes;
-						}
-
-						// switch to the other cache
-						if (CacheStart[CurrentCache] != -1)
-						{
-							CurrentCache++;
-							CurrentCache %= 2;
-						}
-
-						int64 SizeToRead = GBufferCacheSize;
-
-						if (FilePos + SizeToRead > FileSize)
-						{
-							SizeToRead = FileSize-FilePos;
-						}
-
-						if (Network.SendSeekMessage(HandleId, FilePos))
-						{
-							Result = Network.SendReadMessage(HandleId, BufferCache[CurrentCache], SizeToRead);
-						}
-
-						if (Result)
-						{
-							CacheStart[CurrentCache] = FilePos;
-							CacheEnd[CurrentCache] = FilePos + SizeToRead;
-						}
-					}
-
-					// copy from the cache to the destination
-					if (Result)
-					{
-						FMemory::Memcpy(Destination, BufferCache[CurrentCache]+(FilePos-CacheStart[CurrentCache]), BytesToRead);
-						FilePos += BytesToRead;
-					}
-				}
-			}
-		}
-
-		return Result;
-	}
-
-	virtual bool Write(const uint8* Source, int64 BytesToWrite) override
-	{
-		bool Result = false;
-
-		if (bWritable && BytesToWrite >= 0)
-		{
-			if (BytesToWrite == 0)
-			{
-				Result = true;
-			}
-			else
-			{
-				Result = Network.SendWriteMessage(HandleId, Source, BytesToWrite);
-
-				if (Result)
-				{
-					FilePos += BytesToWrite;
-					FileSize = FMath::Max<int64>(FilePos, FileSize);
-				}
-			}
-		}
-
-		return Result;
-	}
-};
-*/
-
-
 bool FCookedIterativeNetworkFile::ShouldBeUsed(IPlatformFile* Inner, const TCHAR* CmdLine) const
 {
 	bool bResult = FNetworkPlatformFile::ShouldBeUsed(Inner, CmdLine);
@@ -580,137 +305,10 @@ bool FCookedIterativeNetworkFile::ShouldBeUsed(IPlatformFile* Inner, const TCHAR
 	return bResult;
 }
 
-/*
-bool FCookedIterativeNetworkFile::InitializeInternal(IPlatformFile* Inner, const TCHAR* HostIP)
-{
-	// look for the commandline that will read files from over the network
-	if (HostIP == nullptr)
-	{
-		UE_LOG(LogStreamingPlatformFile, Error, TEXT("No Host IP specified in the commandline."));
-		bIsUsable = false;
-
-		return false;
-	}
-
-	// optionally get the port from the command line
-	int32 OverridePort;
-	if (FParse::Value(FCommandLine::Get(), TEXT("fileserverport="), OverridePort))
-	{
-		UE_LOG(LogStreamingPlatformFile, Display, TEXT("Overriding file server port: %d"), OverridePort);
-		FileServerPort = OverridePort;
-	}
-
-	// Send the filenames and timestamps to the server.
-	FNetworkFileArchive Payload(NFS_Messages::GetFileList);
-	FillGetFileList(Payload, true);
-
-	// Send the directories over, and wait for a response.
-	FArrayReader Response;
-
-	if(SendPayloadAndReceiveResponse(Payload,Response))
-	{
-		// Receive the cooked version information.
-		int32 ServerPackageVersion = 0;
-		int32 ServerPackageLicenseeVersion = 0;
-		ProcessServerInitialResponse(Response, ServerPackageVersion, ServerPackageLicenseeVersion);
-
-		// Make sure we can sync a file.
-		FString TestSyncFile = FPaths::Combine(*(FPaths::EngineDir()), TEXT("Config/BaseEngine.ini"));
-		IFileHandle* TestFileHandle = OpenRead(*TestSyncFile);
-		if (TestFileHandle != nullptr)
-		{
-			uint8* FileContents = (uint8*)FMemory::Malloc(TestFileHandle->Size());
-			if (!TestFileHandle->Read(FileContents, TestFileHandle->Size()))
-			{
-				UE_LOG(LogStreamingPlatformFile, Fatal, TEXT("Could not read test file %s."), *TestSyncFile);
-			}
-			FMemory::Free(FileContents);
-			delete TestFileHandle;
-		}
-		else
-		{
-			UE_LOG(LogStreamingPlatformFile, Fatal, TEXT("Could not open test file %s."), *TestSyncFile);
-		}
-
-		FCommandLine::AddToSubprocessCommandline( *FString::Printf( TEXT("-StreamingHostIP=%s"), HostIP ) );
-
-		return true; 
-	}
-
-	return false; 
-}*/
-
-
 FCookedIterativeNetworkFile::~FCookedIterativeNetworkFile()
 {
 }
 
-/*
-bool FStreamingNetworkPlatformFile::IterateDirectory(const TCHAR* InDirectory, IPlatformFile::FDirectoryVisitor& Visitor)
-{
-	FString RelativeDirectory = InDirectory;
-	MakeStandardNetworkFilename(RelativeDirectory);
-	// for .dll, etc searches that don't specify a path, we need to strip off the path
-	// before we send it to the visitor
-	bool bHadNoPath = InDirectory[0] == 0;
-
-	// we loop until this is false
-	bool RetVal = true;
-
-	// Find the directory in TOC
-	FServerTOC::FDirectory* ServerDirectory = ServerFiles.FindDirectory(RelativeDirectory);
-	if (ServerDirectory != nullptr)
-	{
-		// loop over the server files and look if they are in this exact directory
-		for (FServerTOC::FDirectory::TIterator It(*ServerDirectory); It && RetVal == true; ++It)
-		{
-			if (FPaths::GetPath(It.Key()) == RelativeDirectory)
-			{
-				// timestamps of 0 mean directories
-				bool bIsDirectory = It.Value() == 0;
-			
-				// visit (stripping off the path if needed)
-				RetVal = Visitor.Visit(bHadNoPath ? *FPaths::GetCleanFilename(It.Key()) : *It.Key(), bIsDirectory);
-			}
-		}
-	}
-
-	return RetVal;
-}
-
-
-bool FStreamingNetworkPlatformFile::IterateDirectoryRecursively(const TCHAR* InDirectory, IPlatformFile::FDirectoryVisitor& Visitor)
-{
-	FString RelativeDirectory = InDirectory;
-	MakeStandardNetworkFilename(RelativeDirectory);
-	// we loop until this is false
-	bool RetVal = true;
-
-	// loop over the server TOC
-	for (TMap<FString, FServerTOC::FDirectory*>::TIterator DirIt(ServerFiles.Directories); DirIt && RetVal == true; ++DirIt)
-	{
-		if (DirIt.Key().StartsWith(RelativeDirectory))
-		{
-			FServerTOC::FDirectory& ServerDirectory = *DirIt.Value();
-		
-			// loop over the server files and look if they are in this exact directory
-			for (FServerTOC::FDirectory::TIterator It(ServerDirectory); It && RetVal == true; ++It)
-			{
-				// timestamps of 0 mean directories
-				bool bIsDirectory = It.Value() == 0;
-
-				// visit!
-				RetVal = Visitor.Visit(*It.Key(), bIsDirectory);
-			}
-		}
-	}
-
-	return RetVal;
-}
-
-*/
-
-
 /**
  * Module for the streaming file
  */
